fix for issue when trying to run a movie on a single threaded device
TTP337210
#ue4

diff --git a/Engine/Source/Runtime/MoviePlayer/Private/DefaultGameMoviePlayer.cpp b/Engine/Source/Runtime/MoviePlayer/Private/DefaultGameMoviePlayer.cpp
--- a/Engine/Source/Runtime/MoviePlayer/Private/DefaultGameMoviePlayer.cpp
+++ b/Engine/Source/Runtime/MoviePlayer/Private/DefaultGameMoviePlayer.cpp
@@ -123,7 +123,7 @@ void FDefaultGameMoviePlayer::SetupLoadingScreen(const FLoadingScreenAttributes&
 bool FDefaultGameMoviePlayer::PlayMovie()
 {
 	bool bBeganPlaying = false;
-	if (LoadingScreenIsPrepared() && !IsMovieCurrentlyPlaying() )
+	if (LoadingScreenIsPrepared() && !IsMovieCurrentlyPlaying() && FPlatformMisc::NumberOfCores() > 1)
 	{
 		check(LoadingScreenAttributes.IsValid());
 
