Disable movie player when -onethread is on the cmd line.
 * Movie player relies on render thread to function correctly.

diff --git a/Engine/Source/Runtime/MoviePlayer/Private/MoviePlayer.cpp b/Engine/Source/Runtime/MoviePlayer/Private/MoviePlayer.cpp
--- a/Engine/Source/Runtime/MoviePlayer/Private/MoviePlayer.cpp
+++ b/Engine/Source/Runtime/MoviePlayer/Private/MoviePlayer.cpp
@@ -76,5 +76,5 @@ TSharedPtr<IGameMoviePlayer> GetMoviePlayer()
 
 bool IsMoviePlayerEnabled()
 {
-	return !GIsEditor && !IsRunningDedicatedServer() && !IsRunningCommandlet();
+	return !GIsEditor && !IsRunningDedicatedServer() && !IsRunningCommandlet() && !FParse::Param(FCommandLine::Get(), TEXT("ONETHREAD"));
 }
