fix for splash screen not staying up until first render
#ue4
#ios

diff --git a/Engine/Source/Runtime/Core/Private/IOS/IOSAppDelegate.cpp b/Engine/Source/Runtime/Core/Private/IOS/IOSAppDelegate.cpp
--- a/Engine/Source/Runtime/Core/Private/IOS/IOSAppDelegate.cpp
+++ b/Engine/Source/Runtime/Core/Private/IOS/IOSAppDelegate.cpp
@@ -417,6 +417,10 @@ void InstallSignalHandlers()
     {
         [path setString: [path stringByAppendingPathComponent:@"Default-Portrait.png"]];
     }
+    else if (MainFrame.size.height == 768 && !self.bDeviceInPortraitMode)
+    {
+        [path setString: [path stringByAppendingPathComponent:@"Default-Landscape.png"]];
+    }
     UIImage* image = [[UIImage alloc] initWithContentsOfFile: path];
     [path release];
     UIImage* imageToDisplay = [UIImage imageWithCGImage: [image CGImage] scale: 1.0 orientation: orient];
diff --git a/Engine/Source/Runtime/Core/Private/IOS/IOSPlatformMisc.cpp b/Engine/Source/Runtime/Core/Private/IOS/IOSPlatformMisc.cpp
--- a/Engine/Source/Runtime/Core/Private/IOS/IOSPlatformMisc.cpp
+++ b/Engine/Source/Runtime/Core/Private/IOS/IOSPlatformMisc.cpp
@@ -58,9 +58,9 @@ void FIOSPlatformMisc::PlatformInit()
 	UE_LOG(LogInit, Log, TEXT("Free Memory at startup: %d MB"), GStartupFreeMemoryMB);
 }
 
-void FIOSPlatformMisc::PlatformPostInit(bool IsMoviePlaying)
+void FIOSPlatformMisc::PlatformPostInit(bool ShowSplashScreen)
 {
-    GShowSplashScreen = !IsMoviePlaying;
+    GShowSplashScreen = ShowSplashScreen;
 }
 
 GenericApplication* FIOSPlatformMisc::CreateApplication()
diff --git a/Engine/Source/Runtime/Core/Private/Mac/MacPlatformMisc.cpp b/Engine/Source/Runtime/Core/Private/Mac/MacPlatformMisc.cpp
--- a/Engine/Source/Runtime/Core/Private/Mac/MacPlatformMisc.cpp
+++ b/Engine/Source/Runtime/Core/Private/Mac/MacPlatformMisc.cpp
@@ -250,7 +250,7 @@ void FMacPlatformMisc::PlatformInit()
 	UE_LOG(LogInit, Log, TEXT("Power Source: %s"), GMacAppInfo.RunningOnBattery ? TEXT(kIOPSBatteryPowerValue) : TEXT(kIOPSACPowerValue) );
 }
 
-void FMacPlatformMisc::PlatformPostInit(bool IsMoviePlaying)
+void FMacPlatformMisc::PlatformPostInit(bool ShowSplashScreen)
 {
 	// Setup the app menu in menu bar
 	const bool bIsBundledApp = [[[NSBundle mainBundle] bundlePath] hasSuffix:@".app"];
diff --git a/Engine/Source/Runtime/Core/Public/GenericPlatform/GenericPlatformMisc.h b/Engine/Source/Runtime/Core/Public/GenericPlatform/GenericPlatformMisc.h
--- a/Engine/Source/Runtime/Core/Public/GenericPlatform/GenericPlatformMisc.h
+++ b/Engine/Source/Runtime/Core/Public/GenericPlatform/GenericPlatformMisc.h
@@ -205,7 +205,7 @@ struct CORE_API FGenericPlatformMisc
 	static void PlatformInit()
 	{
 	}
-	static void PlatformPostInit(bool IsMoviePlaying = false)
+	static void PlatformPostInit(bool ShowSplashScreen = false)
 	{
 	}
 
diff --git a/Engine/Source/Runtime/Core/Public/IOS/IOSPlatformMisc.h b/Engine/Source/Runtime/Core/Public/IOS/IOSPlatformMisc.h
--- a/Engine/Source/Runtime/Core/Public/IOS/IOSPlatformMisc.h
+++ b/Engine/Source/Runtime/Core/Public/IOS/IOSPlatformMisc.h
@@ -12,7 +12,7 @@
 struct CORE_API FIOSPlatformMisc : public FGenericPlatformMisc
 {
 	static void PlatformInit();
-    static void PlatformPostInit(bool IsMoviePlaying = false);
+    static void PlatformPostInit(bool ShowSplashScreen = false);
 	static class GenericApplication* CreateApplication();
 	static void GetEnvironmentVariable(const TCHAR* VariableName, TCHAR* Result, int32 ResultLength);
 	static void* GetHardwareWindow();
diff --git a/Engine/Source/Runtime/Core/Public/Mac/MacPlatformMisc.h b/Engine/Source/Runtime/Core/Public/Mac/MacPlatformMisc.h
--- a/Engine/Source/Runtime/Core/Public/Mac/MacPlatformMisc.h
+++ b/Engine/Source/Runtime/Core/Public/Mac/MacPlatformMisc.h
@@ -13,7 +13,7 @@ struct CORE_API FMacPlatformMisc : public FGenericPlatformMisc
 {
 	static void PlatformPreInit();
 	static void PlatformInit();
-	static void PlatformPostInit(bool IsMoviePlaying = false);
+	static void PlatformPostInit(bool ShowSplashScreen = false);
 	static class GenericApplication* CreateApplication();
 	static void GetEnvironmentVariable(const TCHAR* VariableName, TCHAR* Result, int32 ResultLength);
 
diff --git a/Engine/Source/Runtime/Launch/Private/LaunchEngineLoop.cpp b/Engine/Source/Runtime/Launch/Private/LaunchEngineLoop.cpp
--- a/Engine/Source/Runtime/Launch/Private/LaunchEngineLoop.cpp
+++ b/Engine/Source/Runtime/Launch/Private/LaunchEngineLoop.cpp
@@ -1317,13 +1317,13 @@ int32 FEngineLoop::PreInit( const TCHAR* CmdLine )
 		GetMoviePlayer()->PlayMovie();
 
         // do any post appInit processing, before the render thread is started.
-        FPlatformMisc::PlatformPostInit(GetMoviePlayer()->IsMovieCurrentlyPlaying());
+        FPlatformMisc::PlatformPostInit(!GetMoviePlayer()->IsMovieCurrentlyPlaying());
     }
     else
 #endif
     {
         // do any post appInit processing, before the render thread is started.
-        FPlatformMisc::PlatformPostInit(false);
+        FPlatformMisc::PlatformPostInit(true);
     }
 
 	if (GUseThreadedRendering)
diff --git a/Engine/Source/Runtime/MoviePlayer/Private/DefaultGameMoviePlayer.cpp b/Engine/Source/Runtime/MoviePlayer/Private/DefaultGameMoviePlayer.cpp
--- a/Engine/Source/Runtime/MoviePlayer/Private/DefaultGameMoviePlayer.cpp
+++ b/Engine/Source/Runtime/MoviePlayer/Private/DefaultGameMoviePlayer.cpp
@@ -131,17 +131,21 @@ bool FDefaultGameMoviePlayer::PlayMovie()
 		LoadingIsDone.Set(0);
 		bUserCalledFinish = false;
 		
+        bool bInitialized = true;
 		if (MovieStreamingIsPrepared())
 		{
-			MovieStreamer->Init(LoadingScreenAttributes.MoviePaths);
+			bInitialized = MovieStreamer->Init(LoadingScreenAttributes.MoviePaths);
 		}
-		LoadingScreenWidgetHolder->SetContent(LoadingScreenAttributes.WidgetLoadingScreen.IsValid() ? LoadingScreenAttributes.WidgetLoadingScreen.ToSharedRef() : SNullWidget::NullWidget);
-		LoadingScreenWindowPtr.Pin()->SetContent(LoadingScreenContents.ToSharedRef());
+        if (bInitialized)
+        {
+            LoadingScreenWidgetHolder->SetContent(LoadingScreenAttributes.WidgetLoadingScreen.IsValid() ? LoadingScreenAttributes.WidgetLoadingScreen.ToSharedRef() : SNullWidget::NullWidget);
+            LoadingScreenWindowPtr.Pin()->SetContent(LoadingScreenContents.ToSharedRef());
 		
-		SyncMechanism = new FSlateLoadingSynchronizationMechanism();
-		SyncMechanism->Initialize();
+            SyncMechanism = new FSlateLoadingSynchronizationMechanism();
+            SyncMechanism->Initialize();
 
-		bBeganPlaying = true;
+            bBeganPlaying = true;
+        }
 	}
 
 	return bBeganPlaying;
diff --git a/Engine/Source/Runtime/MoviePlayer/Public/MoviePlayer.h b/Engine/Source/Runtime/MoviePlayer/Public/MoviePlayer.h
--- a/Engine/Source/Runtime/MoviePlayer/Public/MoviePlayer.h
+++ b/Engine/Source/Runtime/MoviePlayer/Public/MoviePlayer.h
@@ -48,7 +48,7 @@ public:
 	 * Initializes this movie streamer with all the movie paths (ordered) we want to play
 	 * Movie paths are local to the current game's Content/Movies/ directory.
 	 */
-	virtual void Init(const TArray<FString>& MoviePaths) = 0;
+	virtual bool Init(const TArray<FString>& MoviePaths) = 0;
 	
 	/** Forces the movie streamer to cancel what it's streaming and close */
 	virtual void ForceCompletion() = 0;
