Add a way to add a minimum display time to a Slate loading screen

#coderevew joel.crabbe

diff --git a/Engine/Source/Runtime/MoviePlayer/Private/DefaultGameMoviePlayer.cpp b/Engine/Source/Runtime/MoviePlayer/Private/DefaultGameMoviePlayer.cpp
--- a/Engine/Source/Runtime/MoviePlayer/Private/DefaultGameMoviePlayer.cpp
+++ b/Engine/Source/Runtime/MoviePlayer/Private/DefaultGameMoviePlayer.cpp
@@ -32,6 +32,7 @@ FDefaultGameMoviePlayer::FDefaultGameMoviePlayer()
 	, LoadingIsDone(0)
 	, bUserCalledFinish(false)
 	, LoadingScreenAttributes()
+	, LastPlayTime(0.0)
 {
 }
 
@@ -131,6 +132,8 @@ bool FDefaultGameMoviePlayer::PlayMovie()
 		LoadingIsDone.Set(0);
 		bUserCalledFinish = false;
 		
+		LastPlayTime = FPlatformTime::Seconds();
+
         bool bInitialized = true;
 		if (MovieStreamingIsPrepared())
 		{
@@ -153,21 +156,29 @@ bool FDefaultGameMoviePlayer::PlayMovie()
 
 void FDefaultGameMoviePlayer::StopMovie()
 {
+	LastPlayTime = 0;
 	bUserCalledFinish = true;
 }
 
 void FDefaultGameMoviePlayer::WaitForMovieToFinish()
 {
-	if (LoadingScreenIsPrepared() && IsMovieCurrentlyPlaying())
+	const bool bEnforceMinimumTime = LoadingScreenAttributes.MinimumLoadingScreenDisplayTime >= 0.0f;
+
+	if (LoadingScreenIsPrepared() && ( IsMovieCurrentlyPlaying() || !bEnforceMinimumTime ) )
 	{
 		SyncMechanism->DestroySlateThread();
 		delete SyncMechanism;
 		SyncMechanism = NULL;
-		
-		LoadingIsDone.Set(1);
+
+		if( !bEnforceMinimumTime )
+		{
+			LoadingIsDone.Set(1);
+		}
 		
 		const bool bAutoCompleteWhenLoadingCompletes = LoadingScreenAttributes.bAutoCompleteWhenLoadingCompletes;
-		while (!IsMovieStreamingFinished() && !bAutoCompleteWhenLoadingCompletes && !bUserCalledFinish)
+	
+		// Continue to wait until the user calls finish (if enabled) or when loading completes or the minimum enforced time (if any) has been reached.
+		while ( !bUserCalledFinish && ( (!bEnforceMinimumTime && !IsMovieStreamingFinished() && !bAutoCompleteWhenLoadingCompletes ) || ( bEnforceMinimumTime &&  (FPlatformTime::Seconds() - LastPlayTime) < LoadingScreenAttributes.MinimumLoadingScreenDisplayTime ) ) )
 		{
 			if (FSlateApplication::IsInitialized())
 			{
@@ -179,12 +190,16 @@ void FDefaultGameMoviePlayer::WaitForMovieToFinish()
 			}
 		}
 
+
+		LoadingIsDone.Set(1);
+
 		MovieStreamingIsDone.Set(1);
 		if( MovieStreamer.IsValid() )
 		{
 			MovieStreamer->ForceCompletion();
 		}
 
+		LastPlayTime = 0;
 		FlushRenderingCommands();
 
 		// Allow the movie streamer to clean up any resources it uses once there are no movies to play.
diff --git a/Engine/Source/Runtime/MoviePlayer/Private/DefaultGameMoviePlayer.h b/Engine/Source/Runtime/MoviePlayer/Private/DefaultGameMoviePlayer.h
--- a/Engine/Source/Runtime/MoviePlayer/Private/DefaultGameMoviePlayer.h
+++ b/Engine/Source/Runtime/MoviePlayer/Private/DefaultGameMoviePlayer.h
@@ -85,6 +85,8 @@ private:
 	/** Attributes of the loading screen we are currently displaying */
 	FLoadingScreenAttributes LoadingScreenAttributes;
 	
+	/** The last time a movie was started */
+	double LastPlayTime;
 private:
 	/** Singleton handle */
 	static TSharedPtr<FDefaultGameMoviePlayer> MoviePlayer;
diff --git a/Engine/Source/Runtime/MoviePlayer/Public/MoviePlayer.h b/Engine/Source/Runtime/MoviePlayer/Public/MoviePlayer.h
--- a/Engine/Source/Runtime/MoviePlayer/Public/MoviePlayer.h
+++ b/Engine/Source/Runtime/MoviePlayer/Public/MoviePlayer.h
@@ -74,7 +74,8 @@ public:
 struct MOVIEPLAYER_API FLoadingScreenAttributes
 {
 	FLoadingScreenAttributes()
-		: bAutoCompleteWhenLoadingCompletes(true)
+		: MinimumLoadingScreenDisplayTime(0.0f)
+		, bAutoCompleteWhenLoadingCompletes(true)
 		, bMoviesAreSkippable(true) {}
 
 	/** The widget to be displayed on top of the movie or simply standalone if there is no movie */
@@ -83,13 +84,15 @@ struct MOVIEPLAYER_API FLoadingScreenAttributes
 	/** The movie paths local to the game's Content/Movies/ directory we will play */
 	TArray<FString> MoviePaths;
 
+	/** The minimum time that a loading screen should be opened for */
+	float MinimumLoadingScreenDisplayTime;
+
 	/** If true, the loading screen will disappear as soon as all movies are played and loading is done */
 	bool bAutoCompleteWhenLoadingCompletes;
 
 	/** If true, movies can be skipped by clicking the loading screen as long as loading is done */
 	bool bMoviesAreSkippable;
 
-
 	/** True if there is either a standalone widget or any movie paths or both */
 	bool IsValid() const;
 
