Integrated GitHub 431 : Closing main window during the intro movie causes engine to hang

diff --git a/Engine/Source/Runtime/MoviePlayer/Private/DefaultGameMoviePlayer.cpp b/Engine/Source/Runtime/MoviePlayer/Private/DefaultGameMoviePlayer.cpp
--- a/Engine/Source/Runtime/MoviePlayer/Private/DefaultGameMoviePlayer.cpp
+++ b/Engine/Source/Runtime/MoviePlayer/Private/DefaultGameMoviePlayer.cpp
@@ -182,6 +182,12 @@ void FDefaultGameMoviePlayer::WaitForMovieToFinish()
 		{
 			if (FSlateApplication::IsInitialized())
 			{
+				// Break out of the loop if the main window is closed during the movie.
+				if ( !FSlateApplication::Get().GetActiveTopLevelWindow().IsValid() )
+				{
+					break;
+				}
+
 				FPlatformMisc::PumpMessages(true);
 				FSlateApplication::Get().Tick();
 				
@@ -391,4 +397,4 @@ void FDefaultGameMoviePlayer::OnPreLoadMap()
 void FDefaultGameMoviePlayer::OnPostLoadMap()
 {
 	WaitForMovieToFinish();
-}
\ No newline at end of file
+}
