- Fix movie player shutdown potentially unregistering a rendering tickable on the game theread
- Fixed movie player holding on to slate widgets when no movie is playing

diff --git a/Engine/Source/Runtime/Launch/Private/LaunchEngineLoop.cpp b/Engine/Source/Runtime/Launch/Private/LaunchEngineLoop.cpp
--- a/Engine/Source/Runtime/Launch/Private/LaunchEngineLoop.cpp
+++ b/Engine/Source/Runtime/Launch/Private/LaunchEngineLoop.cpp
@@ -1889,6 +1889,8 @@ void FEngineLoop::Exit()
 	GIsRunning	= 0;
 	GLogConsole	= NULL;
 
+	GetMoviePlayer()->Shutdown();
+
 	// Make sure we're not in the middle of loading something.
 	FlushAsyncLoading();
 
diff --git a/Engine/Source/Runtime/MoviePlayer/Private/DefaultGameMoviePlayer.cpp b/Engine/Source/Runtime/MoviePlayer/Private/DefaultGameMoviePlayer.cpp
--- a/Engine/Source/Runtime/MoviePlayer/Private/DefaultGameMoviePlayer.cpp
+++ b/Engine/Source/Runtime/MoviePlayer/Private/DefaultGameMoviePlayer.cpp
@@ -27,16 +27,26 @@ TSharedPtr<FDefaultGameMoviePlayer> FDefaultGameMoviePlayer::Get()
 }
 
 FDefaultGameMoviePlayer::FDefaultGameMoviePlayer()
-	: SyncMechanism(NULL)
+	: FTickableObjectRenderThread(false)
+	, SyncMechanism(NULL)
 	, MovieStreamingIsDone(0)
 	, LoadingIsDone(0)
 	, bUserCalledFinish(false)
 	, LoadingScreenAttributes()
 	, LastPlayTime(0.0)
+	, bInitialized(false)
 {
 }
 
-FDefaultGameMoviePlayer::~FDefaultGameMoviePlayer() {}
+FDefaultGameMoviePlayer::~FDefaultGameMoviePlayer()
+{
+	if( bInitialized )
+	{
+		// This should not happen if initialize was called correctly.  This is a fallback to ensure that the movie player rendering tickable gets unregistered on the rendering thread correctly
+		Shutdown();
+		FlushRenderingCommands();
+	}
+}
 
 void FDefaultGameMoviePlayer::RegisterMovieStreamer(TSharedPtr<IMovieStreamer> InMovieStreamer)
 {
@@ -50,6 +60,13 @@ void FDefaultGameMoviePlayer::SetSlateRenderer(TSharedPtr<FSlateRenderer> InSlat
 
 void FDefaultGameMoviePlayer::Initialize()
 {
+	ENQUEUE_UNIQUE_RENDER_COMMAND_ONEPARAMETER(RegisterMoviePlayerTickable, FDefaultGameMoviePlayer*, MoviePlayer, this,
+	{
+		MoviePlayer->Register();
+	});
+
+	bInitialized = true;
+
 	// Initialize shaders, because otherwise they might not be guaranteed to exist at this point
 	if (!FPlatformProperties::RequiresCookedData())
 	{
@@ -107,6 +124,34 @@ void FDefaultGameMoviePlayer::Initialize()
 	LoadingScreenWindowPtr = GameWindow;
 }
 
+void FDefaultGameMoviePlayer::Shutdown()
+{
+
+	ENQUEUE_UNIQUE_RENDER_COMMAND_ONEPARAMETER(UnregisterMoviePlayerTickable, FDefaultGameMoviePlayer*, MoviePlayer, this,
+	{
+		MoviePlayer->Unregister();
+	});
+
+	bInitialized = false;
+
+	FCoreDelegates::PreLoadMap.RemoveAll( this );
+
+	LoadingScreenContents.Reset();
+	LoadingScreenWidgetHolder.Reset();
+	LoadingScreenWindowPtr.Reset();
+
+	MovieStreamer.Reset();
+
+	LoadingScreenAttributes = FLoadingScreenAttributes();
+
+	if( SyncMechanism)
+	{
+		SyncMechanism->DestroySlateThread();
+		delete SyncMechanism;
+		SyncMechanism = NULL;
+	}
+
+}
 void FDefaultGameMoviePlayer::PassLoadingScreenWindowBackToGame() const
 {
 	auto GameEngine = Cast<UGameEngine>(GEngine);
@@ -158,6 +203,7 @@ void FDefaultGameMoviePlayer::StopMovie()
 {
 	LastPlayTime = 0;
 	bUserCalledFinish = true;
+	LoadingScreenWidgetHolder->SetContent( SNullWidget::NullWidget );
 }
 
 void FDefaultGameMoviePlayer::WaitForMovieToFinish()
@@ -196,6 +242,7 @@ void FDefaultGameMoviePlayer::WaitForMovieToFinish()
 			}
 		}
 
+		LoadingScreenWidgetHolder->SetContent( SNullWidget::NullWidget );
 
 		LoadingIsDone.Set(1);
 
diff --git a/Engine/Source/Runtime/MoviePlayer/Private/DefaultGameMoviePlayer.h b/Engine/Source/Runtime/MoviePlayer/Private/DefaultGameMoviePlayer.h
--- a/Engine/Source/Runtime/MoviePlayer/Private/DefaultGameMoviePlayer.h
+++ b/Engine/Source/Runtime/MoviePlayer/Private/DefaultGameMoviePlayer.h
@@ -17,6 +17,7 @@ public:
 	virtual void RegisterMovieStreamer(TSharedPtr<IMovieStreamer> InMovieStreamer) override;
 	virtual void SetSlateRenderer(TSharedPtr<FSlateRenderer> InSlateRenderer) override;
 	virtual void Initialize() override;
+	virtual void Shutdown() override;
 	virtual void PassLoadingScreenWindowBackToGame() const override;
 	virtual void SetupLoadingScreen(const FLoadingScreenAttributes& LoadingScreenAttributes) override;
 	virtual bool PlayMovie() override;
@@ -87,6 +88,9 @@ private:
 	
 	/** The last time a movie was started */
 	double LastPlayTime;
+
+	/** True if the movie player has been initialized */
+	bool bInitialized;
 private:
 	/** Singleton handle */
 	static TSharedPtr<FDefaultGameMoviePlayer> MoviePlayer;
diff --git a/Engine/Source/Runtime/MoviePlayer/Private/NullMoviePlayer.h b/Engine/Source/Runtime/MoviePlayer/Private/NullMoviePlayer.h
--- a/Engine/Source/Runtime/MoviePlayer/Private/NullMoviePlayer.h
+++ b/Engine/Source/Runtime/MoviePlayer/Private/NullMoviePlayer.h
@@ -21,6 +21,7 @@ public:
 	virtual void RegisterMovieStreamer(TSharedPtr<IMovieStreamer> InMovieStreamer) override {}
 	virtual void SetSlateRenderer(TSharedPtr<FSlateRenderer> InSlateRenderer) override {}
 	virtual void Initialize() override {}
+	virtual void Shutdown() override {}
 	virtual void PassLoadingScreenWindowBackToGame() const override {}
 	virtual void SetupLoadingScreen(const FLoadingScreenAttributes& InLoadingScreenAttributes) override {}
 	virtual bool PlayMovie() override { return false;}
diff --git a/Engine/Source/Runtime/MoviePlayer/Public/MoviePlayer.h b/Engine/Source/Runtime/MoviePlayer/Public/MoviePlayer.h
--- a/Engine/Source/Runtime/MoviePlayer/Public/MoviePlayer.h
+++ b/Engine/Source/Runtime/MoviePlayer/Public/MoviePlayer.h
@@ -115,6 +115,9 @@ public:
 	/** Initializes this movie player, creating the startup window and hiding the splash screen. To be called in the launch engine loop */
 	virtual void Initialize() = 0;
 
+	/** Shutsdown the movie player */
+	virtual void Shutdown() = 0;
+
 	/** Passes the loading screen window back to the game to use. For use by the launch engine loop only */
 	virtual void PassLoadingScreenWindowBackToGame() const = 0;
 	
