Changed the loading screen to start in a loaded state.

diff --git a/Engine/Source/Runtime/MoviePlayer/Private/DefaultGameMoviePlayer.cpp b/Engine/Source/Runtime/MoviePlayer/Private/DefaultGameMoviePlayer.cpp
--- a/Engine/Source/Runtime/MoviePlayer/Private/DefaultGameMoviePlayer.cpp
+++ b/Engine/Source/Runtime/MoviePlayer/Private/DefaultGameMoviePlayer.cpp
@@ -30,7 +30,7 @@ FDefaultGameMoviePlayer::FDefaultGameMoviePlayer()
 	: FTickableObjectRenderThread(false)
 	, SyncMechanism(NULL)
 	, MovieStreamingIsDone(0)
-	, LoadingIsDone(0)
+	, LoadingIsDone(1)
 	, bUserCalledFinish(false)
 	, LoadingScreenAttributes()
 	, LastPlayTime(0.0)
