#UE4 - iOS - Fix crash loading ith no movie

diff --git a/Engine/Source/Runtime/MoviePlayer/Private/DefaultGameMoviePlayer.cpp b/Engine/Source/Runtime/MoviePlayer/Private/DefaultGameMoviePlayer.cpp
--- a/Engine/Source/Runtime/MoviePlayer/Private/DefaultGameMoviePlayer.cpp
+++ b/Engine/Source/Runtime/MoviePlayer/Private/DefaultGameMoviePlayer.cpp
@@ -194,6 +194,10 @@ bool FDefaultGameMoviePlayer::PlayMovie()
 
             bBeganPlaying = true;
         }
+		else
+		{
+			LoadingIsDone.Set(1);
+		}
 	}
 
 	return bBeganPlaying;
