Better fix for loading without a movie
#ue4
#ios
#android

diff --git a/Engine/Source/Runtime/MoviePlayer/Private/DefaultGameMoviePlayer.cpp b/Engine/Source/Runtime/MoviePlayer/Private/DefaultGameMoviePlayer.cpp
--- a/Engine/Source/Runtime/MoviePlayer/Private/DefaultGameMoviePlayer.cpp
+++ b/Engine/Source/Runtime/MoviePlayer/Private/DefaultGameMoviePlayer.cpp
@@ -29,7 +29,7 @@ TSharedPtr<FDefaultGameMoviePlayer> FDefaultGameMoviePlayer::Get()
 FDefaultGameMoviePlayer::FDefaultGameMoviePlayer()
 	: FTickableObjectRenderThread(false)
 	, SyncMechanism(NULL)
-	, MovieStreamingIsDone(0)
+	, MovieStreamingIsDone(1)
 	, LoadingIsDone(1)
 	, bUserCalledFinish(false)
 	, LoadingScreenAttributes()
@@ -169,12 +169,10 @@ void FDefaultGameMoviePlayer::SetupLoadingScreen(const FLoadingScreenAttributes&
 bool FDefaultGameMoviePlayer::PlayMovie()
 {
 	bool bBeganPlaying = false;
+
 	if (LoadingScreenIsPrepared() && !IsMovieCurrentlyPlaying() && FPlatformMisc::NumberOfCores() > 1)
 	{
 		check(LoadingScreenAttributes.IsValid());
-
-		MovieStreamingIsDone.Set(MovieStreamingIsPrepared() ? 0 : 1);
-		LoadingIsDone.Set(0);
 		bUserCalledFinish = false;
 		
 		LastPlayTime = FPlatformTime::Seconds();
@@ -186,7 +184,10 @@ bool FDefaultGameMoviePlayer::PlayMovie()
 		}
         if (bInitialized)
         {
-            LoadingScreenWidgetHolder->SetContent(LoadingScreenAttributes.WidgetLoadingScreen.IsValid() ? LoadingScreenAttributes.WidgetLoadingScreen.ToSharedRef() : SNullWidget::NullWidget);
+			MovieStreamingIsDone.Set(MovieStreamingIsPrepared() ? 0 : 1);
+			LoadingIsDone.Set(0);
+			
+			LoadingScreenWidgetHolder->SetContent(LoadingScreenAttributes.WidgetLoadingScreen.IsValid() ? LoadingScreenAttributes.WidgetLoadingScreen.ToSharedRef() : SNullWidget::NullWidget);
             LoadingScreenWindowPtr.Pin()->SetContent(LoadingScreenContents.ToSharedRef());
 		
             SyncMechanism = new FSlateLoadingSynchronizationMechanism();
@@ -194,10 +195,6 @@ bool FDefaultGameMoviePlayer::PlayMovie()
 
             bBeganPlaying = true;
         }
-		else
-		{
-			LoadingIsDone.Set(1);
-		}
 	}
 
 	return bBeganPlaying;
