Change MinimumLoadingScreenDisplayTime default to -1 to allow wait for movie completion to work
#ue4

diff --git a/Engine/Source/Runtime/MoviePlayer/Public/MoviePlayer.h b/Engine/Source/Runtime/MoviePlayer/Public/MoviePlayer.h
--- a/Engine/Source/Runtime/MoviePlayer/Public/MoviePlayer.h
+++ b/Engine/Source/Runtime/MoviePlayer/Public/MoviePlayer.h
@@ -74,7 +74,7 @@ public:
 struct MOVIEPLAYER_API FLoadingScreenAttributes
 {
 	FLoadingScreenAttributes()
-		: MinimumLoadingScreenDisplayTime(0.0f)
+		: MinimumLoadingScreenDisplayTime(-1.0f)
 		, bAutoCompleteWhenLoadingCompletes(true)
 		, bMoviesAreSkippable(true) {}
 
