Call PollGameDeviceState during WaitForMovieToFinish so touch/controller events are processed allowing MoviesAreSkippable to work.  Check LoadingScreenWindowPtr instead of ActiveTopLevelWindow to break out of wait if closed.
#ue4

diff --git a/Engine/Source/Runtime/MoviePlayer/Private/DefaultGameMoviePlayer.cpp b/Engine/Source/Runtime/MoviePlayer/Private/DefaultGameMoviePlayer.cpp
--- a/Engine/Source/Runtime/MoviePlayer/Private/DefaultGameMoviePlayer.cpp
+++ b/Engine/Source/Runtime/MoviePlayer/Private/DefaultGameMoviePlayer.cpp
@@ -223,23 +223,30 @@ void FDefaultGameMoviePlayer::WaitForMovieToFinish()
 		}
 		
 		const bool bAutoCompleteWhenLoadingCompletes = LoadingScreenAttributes.bAutoCompleteWhenLoadingCompletes;
-	
+
+		FSlateApplication& SlateApp = FSlateApplication::Get();
+
 		// Continue to wait until the user calls finish (if enabled) or when loading completes or the minimum enforced time (if any) has been reached.
 		while ( !bUserCalledFinish && ( (!bEnforceMinimumTime && !IsMovieStreamingFinished() && !bAutoCompleteWhenLoadingCompletes ) || ( bEnforceMinimumTime &&  (FPlatformTime::Seconds() - LastPlayTime) < LoadingScreenAttributes.MinimumLoadingScreenDisplayTime ) ) )
 		{
 			if (FSlateApplication::IsInitialized())
 			{
 				// Break out of the loop if the main window is closed during the movie.
-				if ( !FSlateApplication::Get().GetActiveTopLevelWindow().IsValid() )
+				if ( !LoadingScreenWindowPtr.IsValid() )
 				{
 					break;
 				}
 
 				FPlatformMisc::PumpMessages(true);
-				FSlateApplication::Get().Tick();
-				
+
+				SlateApp.PollGameDeviceState();
+				// Gives widgets a chance to process any accumulated input
+				SlateApp.FinishedInputThisFrame();
+
+				SlateApp.Tick();
+
 				// Synchronize the game thread and the render thread so that the render thread doesn't get too far behind.
-				FSlateApplication::Get().GetRenderer()->Sync();
+				SlateApp.GetRenderer()->Sync();
 			}
 		}
 
