Fix a crash trying to access null 'SyncMechanism' in FDefaultGameMoviePlayer

diff --git a/Engine/Source/Runtime/MoviePlayer/Private/DefaultGameMoviePlayer.cpp b/Engine/Source/Runtime/MoviePlayer/Private/DefaultGameMoviePlayer.cpp
--- a/Engine/Source/Runtime/MoviePlayer/Private/DefaultGameMoviePlayer.cpp
+++ b/Engine/Source/Runtime/MoviePlayer/Private/DefaultGameMoviePlayer.cpp
@@ -213,9 +213,12 @@ void FDefaultGameMoviePlayer::WaitForMovieToFinish()
 
 	if (LoadingScreenIsPrepared() && ( IsMovieCurrentlyPlaying() || !bEnforceMinimumTime ) )
 	{
-		SyncMechanism->DestroySlateThread();
-		delete SyncMechanism;
-		SyncMechanism = NULL;
+		if (SyncMechanism)
+		{
+			SyncMechanism->DestroySlateThread();
+			delete SyncMechanism;
+			SyncMechanism = NULL;
+		}
 
 		if( !bEnforceMinimumTime )
 		{
