Added delegate callback to the movie player when it finishes playing movies and can safely return to gameplay

diff --git a/Engine/Source/Runtime/MoviePlayer/Private/DefaultGameMoviePlayer.cpp b/Engine/Source/Runtime/MoviePlayer/Private/DefaultGameMoviePlayer.cpp
--- a/Engine/Source/Runtime/MoviePlayer/Private/DefaultGameMoviePlayer.cpp
+++ b/Engine/Source/Runtime/MoviePlayer/Private/DefaultGameMoviePlayer.cpp
@@ -265,17 +265,18 @@ void FDefaultGameMoviePlayer::WaitForMovieToFinish()
 
 		LastPlayTime = 0;
 		FlushRenderingCommands();
-
+		
 		// Allow the movie streamer to clean up any resources it uses once there are no movies to play.
 		if( MovieStreamer.IsValid() )
 		{
 			MovieStreamer->Cleanup();
 		}
 	
-
 		// Finally, clear out the loading screen attributes, forcing users to always
 		// explicitly set the loading screen they want (rather than have stale loading screens)
 		LoadingScreenAttributes = FLoadingScreenAttributes();
+
+		OnMoviePlaybackFinishedDelegate.Broadcast();
 	}
 	else
 	{	
diff --git a/Engine/Source/Runtime/MoviePlayer/Private/DefaultGameMoviePlayer.h b/Engine/Source/Runtime/MoviePlayer/Private/DefaultGameMoviePlayer.h
--- a/Engine/Source/Runtime/MoviePlayer/Private/DefaultGameMoviePlayer.h
+++ b/Engine/Source/Runtime/MoviePlayer/Private/DefaultGameMoviePlayer.h
@@ -27,7 +27,7 @@ public:
 	virtual bool IsMovieCurrentlyPlaying() const override;
 	virtual bool LoadingScreenIsPrepared() const override;
 	virtual void SetupLoadingScreenFromIni() override;
-
+	virtual FOnMoviePlaybackFinished& OnMoviePlaybackFinished() override { return OnMoviePlaybackFinishedDelegate; }
 	/** FTickableObjectRenderThread interface */
 	virtual void Tick( float DeltaTime ) override;
 	virtual TStatId GetStatId() const override;
@@ -86,6 +86,8 @@ private:
 	/** Attributes of the loading screen we are currently displaying */
 	FLoadingScreenAttributes LoadingScreenAttributes;
 	
+	FOnMoviePlaybackFinished OnMoviePlaybackFinishedDelegate;
+
 	/** The last time a movie was started */
 	double LastPlayTime;
 
diff --git a/Engine/Source/Runtime/MoviePlayer/Public/MoviePlayer.h b/Engine/Source/Runtime/MoviePlayer/Public/MoviePlayer.h
--- a/Engine/Source/Runtime/MoviePlayer/Public/MoviePlayer.h
+++ b/Engine/Source/Runtime/MoviePlayer/Public/MoviePlayer.h
@@ -101,7 +101,6 @@ struct MOVIEPLAYER_API FLoadingScreenAttributes
 };
 
 
-
 /** An interface to the movie player we will use for loading screens and gameplay movies */
 class IGameMoviePlayer
 {
@@ -150,6 +149,9 @@ public:
 	/** Sets up an FLoadingScreenAttributes from the game's engine.ini, then calls the virtual SetupLoadingScreen */
 	virtual void SetupLoadingScreenFromIni() = 0;
 
+	DECLARE_EVENT(IGameMoviePlayer, FOnMoviePlaybackFinished)
+	virtual FOnMoviePlaybackFinished& OnMoviePlaybackFinished() = 0;
+
 	virtual ~IGameMoviePlayer() {}
 };
 
