Fix build

diff --git a/Engine/Source/Runtime/MoviePlayer/Private/NullMoviePlayer.h b/Engine/Source/Runtime/MoviePlayer/Private/NullMoviePlayer.h
--- a/Engine/Source/Runtime/MoviePlayer/Private/NullMoviePlayer.h
+++ b/Engine/Source/Runtime/MoviePlayer/Private/NullMoviePlayer.h
@@ -31,11 +31,12 @@ public:
 	virtual bool IsMovieCurrentlyPlaying() const override  {return false;}
 	virtual bool LoadingScreenIsPrepared() const override {return false;}
 	virtual void SetupLoadingScreenFromIni() override {}
-
+	virtual FOnMoviePlaybackFinished& OnMoviePlaybackFinished() override { return OnMoviePlaybackFinishedDelegate; }
 private:
 	FNullGameMoviePlayer() {}
 	
 private:
 	/** Singleton handle */
 	static TSharedPtr<FNullGameMoviePlayer> MoviePlayer;
+	FOnMoviePlaybackFinished OnMoviePlaybackFinishedDelegate;
 };
