Fix XB1 compile error

diff --git a/Engine/Source/Runtime/MoviePlayer/Private/DefaultGameMoviePlayer.cpp b/Engine/Source/Runtime/MoviePlayer/Private/DefaultGameMoviePlayer.cpp
--- a/Engine/Source/Runtime/MoviePlayer/Private/DefaultGameMoviePlayer.cpp
+++ b/Engine/Source/Runtime/MoviePlayer/Private/DefaultGameMoviePlayer.cpp
@@ -276,7 +276,7 @@ void FDefaultGameMoviePlayer::WaitForMovieToFinish()
 		// explicitly set the loading screen they want (rather than have stale loading screens)
 		LoadingScreenAttributes = FLoadingScreenAttributes();
 
-		OnMoviePlaybackFinishedDelegate.Broadcast();
+		BroadcastMoviePlaybackFinished();
 	}
 	else
 	{	
diff --git a/Engine/Source/Runtime/MoviePlayer/Public/MoviePlayer.h b/Engine/Source/Runtime/MoviePlayer/Public/MoviePlayer.h
--- a/Engine/Source/Runtime/MoviePlayer/Public/MoviePlayer.h
+++ b/Engine/Source/Runtime/MoviePlayer/Public/MoviePlayer.h
@@ -151,6 +151,8 @@ public:
 
 	DECLARE_EVENT(IGameMoviePlayer, FOnMoviePlaybackFinished)
 	virtual FOnMoviePlaybackFinished& OnMoviePlaybackFinished() = 0;
+	
+	void BroadcastMoviePlaybackFinished() { OnMoviePlaybackFinished().Broadcast(); }
 
 	virtual ~IGameMoviePlayer() {}
 };
