Potential fix for race condition in loading screen movie player

diff --git a/Engine/Source/Runtime/MoviePlayer/Private/DefaultGameMoviePlayer.cpp b/Engine/Source/Runtime/MoviePlayer/Private/DefaultGameMoviePlayer.cpp
--- a/Engine/Source/Runtime/MoviePlayer/Private/DefaultGameMoviePlayer.cpp
+++ b/Engine/Source/Runtime/MoviePlayer/Private/DefaultGameMoviePlayer.cpp
@@ -258,14 +258,14 @@ void FDefaultGameMoviePlayer::WaitForMovieToFinish()
 		LoadingIsDone.Set(1);
 
 		MovieStreamingIsDone.Set(1);
+
+		FlushRenderingCommands();
+
 		if( MovieStreamer.IsValid() )
 		{
 			MovieStreamer->ForceCompletion();
 		}
 
-		LastPlayTime = 0;
-		FlushRenderingCommands();
-		
 		// Allow the movie streamer to clean up any resources it uses once there are no movies to play.
 		if( MovieStreamer.IsValid() )
 		{
