Dont't force kill a thread

diff --git a/Engine/Source/Runtime/MoviePlayer/Private/MoviePlayerThreading.cpp b/Engine/Source/Runtime/MoviePlayer/Private/MoviePlayerThreading.cpp
--- a/Engine/Source/Runtime/MoviePlayer/Private/MoviePlayerThreading.cpp
+++ b/Engine/Source/Runtime/MoviePlayer/Private/MoviePlayerThreading.cpp
@@ -64,7 +64,6 @@ void FSlateLoadingSynchronizationMechanism::DestroySlateThread()
 
 		MainLoop.BlockUntilUnlocked();
 
-		SlateLoadingThread->Kill();
 		delete SlateLoadingThread;
 		delete SlateRunnableTask;
 		SlateLoadingThread = NULL;
