Fixed issue where ShooterGame startup movie was not visible in the viewport.  FDefaultGameMoviePlayer needs its viewport to be set as 'active'.

diff --git a/Engine/Source/Runtime/MoviePlayer/Private/DefaultGameMoviePlayer.cpp b/Engine/Source/Runtime/MoviePlayer/Private/DefaultGameMoviePlayer.cpp
--- a/Engine/Source/Runtime/MoviePlayer/Private/DefaultGameMoviePlayer.cpp
+++ b/Engine/Source/Runtime/MoviePlayer/Private/DefaultGameMoviePlayer.cpp
@@ -115,6 +115,8 @@ void FDefaultGameMoviePlayer::Initialize()
 		MovieViewport->SetViewportInterface( MovieStreamer->GetViewportInterface().ToSharedRef() );
 	}
 	
+	MovieViewport->SetActive(true);
+
 	// Register the movie viewport so that it can receive user input.
 	if (!FPlatformProperties::SupportsWindowedMode())
 	{
