Potential fix for a deadlock

diff --git a/Engine/Source/Runtime/MoviePlayer/Private/MoviePlayerThreading.cpp b/Engine/Source/Runtime/MoviePlayer/Private/MoviePlayerThreading.cpp
--- a/Engine/Source/Runtime/MoviePlayer/Private/MoviePlayerThreading.cpp
+++ b/Engine/Source/Runtime/MoviePlayer/Private/MoviePlayerThreading.cpp
@@ -29,6 +29,7 @@ public:
 	/** FRunnable interface */
 	virtual bool Init() override;
 	virtual uint32 Run() override;
+	virtual void Stop() override;
 private:
 	/** Hold a handle to our parent sync mechanism which handles all of our threading locks */
 	class FSlateLoadingSynchronizationMechanism* SyncMechanism;
@@ -163,3 +164,9 @@ uint32 FSlateLoadingThreadTask::Run()
 
 	return 0;
 }
+
+void FSlateLoadingThreadTask::Stop()
+{
+	SyncMechanism->ResetSlateDrawPassEnqueued();
+	SyncMechanism->ResetSlateMainLoopRunning();
+}
