UE-15995/Shootergame hangs on load - Shutdown should no longer dead-lock

diff --git a/Engine/Source/Runtime/MoviePlayer/Private/DefaultGameMoviePlayer.cpp b/Engine/Source/Runtime/MoviePlayer/Private/DefaultGameMoviePlayer.cpp
--- a/Engine/Source/Runtime/MoviePlayer/Private/DefaultGameMoviePlayer.cpp
+++ b/Engine/Source/Runtime/MoviePlayer/Private/DefaultGameMoviePlayer.cpp
@@ -128,6 +128,8 @@ void FDefaultGameMoviePlayer::Initialize()
 
 void FDefaultGameMoviePlayer::Shutdown()
 {
+	StopMovie();
+	WaitForMovieToFinish();
 
 	ENQUEUE_UNIQUE_RENDER_COMMAND_ONEPARAMETER(UnregisterMoviePlayerTickable, FDefaultGameMoviePlayer*, MoviePlayer, this,
 	{
