UE4 - flush rendering thread on movie finish wait to avoid several buffering problems

diff --git a/Engine/Source/Runtime/MoviePlayer/Private/DefaultGameMoviePlayer.cpp b/Engine/Source/Runtime/MoviePlayer/Private/DefaultGameMoviePlayer.cpp
--- a/Engine/Source/Runtime/MoviePlayer/Private/DefaultGameMoviePlayer.cpp
+++ b/Engine/Source/Runtime/MoviePlayer/Private/DefaultGameMoviePlayer.cpp
@@ -253,6 +253,8 @@ void FDefaultGameMoviePlayer::WaitForMovieToFinish()
 
 				// Synchronize the game thread and the render thread so that the render thread doesn't get too far behind.
 				SlateApp.GetRenderer()->Sync();
+
+				FlushRenderingCommands();
 			}
 		}
 
