UE4 - flush the RHI thread and RHI resources in the movie loop

diff --git a/Engine/Source/Runtime/MoviePlayer/Private/MoviePlayerThreading.cpp b/Engine/Source/Runtime/MoviePlayer/Private/MoviePlayerThreading.cpp
--- a/Engine/Source/Runtime/MoviePlayer/Private/MoviePlayerThreading.cpp
+++ b/Engine/Source/Runtime/MoviePlayer/Private/MoviePlayerThreading.cpp
@@ -130,6 +130,13 @@ void FSlateLoadingSynchronizationMechanism::SlateThreadRunMainLoop()
 			//FSlateApplication::Get().PumpMessages();
 			FSlateApplication::Get().Tick();
 			SetSlateDrawPassEnqueued();
+
+			ENQUEUE_UNIQUE_RENDER_COMMAND(
+				FlushPendingDeleteRHIResources,
+				{
+					GRHICommandList.GetImmediateCommandList().ImmediateFlush(EImmediateFlushType::FlushRHIThreadFlushResources);
+				}
+			);
 		}
 
 		LastTime = CurrentTime;
