UE4 - fix threading issue with movie playback and no rhi thread

diff --git a/Engine/Source/Runtime/MoviePlayer/Private/DefaultGameMoviePlayer.cpp b/Engine/Source/Runtime/MoviePlayer/Private/DefaultGameMoviePlayer.cpp
--- a/Engine/Source/Runtime/MoviePlayer/Private/DefaultGameMoviePlayer.cpp
+++ b/Engine/Source/Runtime/MoviePlayer/Private/DefaultGameMoviePlayer.cpp
@@ -333,6 +333,7 @@ void FDefaultGameMoviePlayer::Tick( float DeltaTime )
 				TickStreamer(DeltaTime);
 				RendererPtr.Pin()->DrawWindows();
 				SyncMechanism->ResetSlateDrawPassEnqueued();
+				GRHICommandList.GetImmediateCommandList().ImmediateFlush(EImmediateFlushType::FlushRHIThreadFlushResources);
 			}
 		}
 	}
diff --git a/Engine/Source/Runtime/MoviePlayer/Private/MoviePlayerThreading.cpp b/Engine/Source/Runtime/MoviePlayer/Private/MoviePlayerThreading.cpp
--- a/Engine/Source/Runtime/MoviePlayer/Private/MoviePlayerThreading.cpp
+++ b/Engine/Source/Runtime/MoviePlayer/Private/MoviePlayerThreading.cpp
@@ -130,13 +130,6 @@ void FSlateLoadingSynchronizationMechanism::SlateThreadRunMainLoop()
 			//FSlateApplication::Get().PumpMessages();
 			FSlateApplication::Get().Tick();
 			SetSlateDrawPassEnqueued();
-
-			ENQUEUE_UNIQUE_RENDER_COMMAND(
-				FlushPendingDeleteRHIResources,
-				{
-					GRHICommandList.GetImmediateCommandList().ImmediateFlush(EImmediateFlushType::FlushRHIThreadFlushResources);
-				}
-			);
 		}
 
 		LastTime = CurrentTime;
diff --git a/Engine/Source/Runtime/RenderCore/Public/RenderingThread.h b/Engine/Source/Runtime/RenderCore/Public/RenderingThread.h
--- a/Engine/Source/Runtime/RenderCore/Public/RenderingThread.h
+++ b/Engine/Source/Runtime/RenderCore/Public/RenderingThread.h
@@ -196,6 +196,7 @@ DECLARE_STATS_GROUP(TEXT("Render Thread Commands"), STATGROUP_RenderThreadComman
 		if(GIsThreadedRendering || !IsInGameThread()) \
 		{ \
 			CheckNotBlockedOnRenderThread(); \
+			check(ENamedThreads::GameThread != ENamedThreads::RenderThread); \
 			TGraphTask<EURCMacro_##TypeName>::CreateTask().ConstructAndDispatchWhenReady(); \
 		} \
 		else \
