Merging using UE4_to_UE4-Orion

CL 2678071
Add a commandline parameter to disable loading screen movies in non-shipping builds

diff --git a/Engine/Source/Runtime/MoviePlayer/Private/MoviePlayer.cpp b/Engine/Source/Runtime/MoviePlayer/Private/MoviePlayer.cpp
--- a/Engine/Source/Runtime/MoviePlayer/Private/MoviePlayer.cpp
+++ b/Engine/Source/Runtime/MoviePlayer/Private/MoviePlayer.cpp
@@ -77,5 +77,11 @@ TSharedPtr<IGameMoviePlayer> GetMoviePlayer()
 
 bool IsMoviePlayerEnabled()
 {
-	return !GIsEditor && !IsRunningDedicatedServer() && !IsRunningCommandlet() && GUseThreadedRendering;
+	bool bEnabled = !GIsEditor && !IsRunningDedicatedServer() && !IsRunningCommandlet() && GUseThreadedRendering;
+	
+#if !UE_BUILD_SHIPPING
+	bEnabled &= !FParse::Param(FCommandLine::Get(), TEXT("NoLoadingScreen"));
+#endif
+
+	return bEnabled;
 }
