Correct movie playback stop condition and add new bWaitForManualStop flag
#ue4

diff --git a/Engine/Source/Runtime/MoviePlayer/Private/DefaultGameMoviePlayer.cpp b/Engine/Source/Runtime/MoviePlayer/Private/DefaultGameMoviePlayer.cpp
--- a/Engine/Source/Runtime/MoviePlayer/Private/DefaultGameMoviePlayer.cpp
+++ b/Engine/Source/Runtime/MoviePlayer/Private/DefaultGameMoviePlayer.cpp
@@ -233,10 +233,11 @@ void FDefaultGameMoviePlayer::WaitForMovieToFinish()
 		}
 		
 		const bool bAutoCompleteWhenLoadingCompletes = LoadingScreenAttributes.bAutoCompleteWhenLoadingCompletes;
+		const bool bWaitForManualStop = LoadingScreenAttributes.bWaitForManualStop;
 
 		FSlateApplication& SlateApp = FSlateApplication::Get();
 		// Continue to wait until the user calls finish (if enabled) or when loading completes or the minimum enforced time (if any) has been reached.
-		while ( !bUserCalledFinish && ( !bAutoCompleteWhenLoadingCompletes || ( !IsMovieStreamingFinished() && ( !bEnforceMinimumTime || (FPlatformTime::Seconds() - LastPlayTime) < LoadingScreenAttributes.MinimumLoadingScreenDisplayTime ) ) ) )
+		while ( (bWaitForManualStop && !bUserCalledFinish) || (!bUserCalledFinish && ((!bEnforceMinimumTime && !IsMovieStreamingFinished() && !bAutoCompleteWhenLoadingCompletes) || (bEnforceMinimumTime && (FPlatformTime::Seconds() - LastPlayTime) < LoadingScreenAttributes.MinimumLoadingScreenDisplayTime))))
 		{
 			if (FSlateApplication::IsInitialized())
 			{
diff --git a/Engine/Source/Runtime/MoviePlayer/Public/MoviePlayer.h b/Engine/Source/Runtime/MoviePlayer/Public/MoviePlayer.h
--- a/Engine/Source/Runtime/MoviePlayer/Public/MoviePlayer.h
+++ b/Engine/Source/Runtime/MoviePlayer/Public/MoviePlayer.h
@@ -76,7 +76,8 @@ struct MOVIEPLAYER_API FLoadingScreenAttributes
 	FLoadingScreenAttributes()
 		: MinimumLoadingScreenDisplayTime(-1.0f)
 		, bAutoCompleteWhenLoadingCompletes(true)
-		, bMoviesAreSkippable(true) {}
+		, bMoviesAreSkippable(true)
+		, bWaitForManualStop(false) {}
 
 	/** The widget to be displayed on top of the movie or simply standalone if there is no movie. */
 	TSharedPtr<class SWidget> WidgetLoadingScreen;
@@ -93,6 +94,9 @@ struct MOVIEPLAYER_API FLoadingScreenAttributes
 	/** If true, movies can be skipped by clicking the loading screen as long as loading is done. */
 	bool bMoviesAreSkippable;
 
+	/** If true, movie playback continues until Stop is called. */
+	bool bWaitForManualStop;
+
 	/** True if there is either a standalone widget or any movie paths or both. */
 	bool IsValid() const;
 
