Loading movies can now be skipped via controller or keyboard.

diff --git a/Engine/Source/Runtime/MoviePlayer/Private/DefaultGameMoviePlayer.cpp b/Engine/Source/Runtime/MoviePlayer/Private/DefaultGameMoviePlayer.cpp
--- a/Engine/Source/Runtime/MoviePlayer/Private/DefaultGameMoviePlayer.cpp
+++ b/Engine/Source/Runtime/MoviePlayer/Private/DefaultGameMoviePlayer.cpp
@@ -15,6 +15,51 @@
 #include "MoviePlayerSettings.h"
 #include "ShaderCompiler.h"
 
+class SDefaultMovieBorder : public SBorder
+{
+public:
+
+	SLATE_BEGIN_ARGS(SDefaultMovieBorder)		
+		: _OnKeyDown()
+	{}
+
+		SLATE_EVENT(FPointerEventHandler, OnMouseButtonDown)
+		SLATE_EVENT(FOnKeyDown, OnKeyDown)
+		SLATE_DEFAULT_SLOT(FArguments, Content)
+
+	SLATE_END_ARGS()
+
+	/**
+	* Construct this widget
+	*
+	* @param	InArgs	The declaration data for this widget
+	*/
+	void Construct(const FArguments& InArgs)
+	{
+		OnKeyDown = InArgs._OnKeyDown;		
+
+		SBorder::Construct(SBorder::FArguments()			
+			.BorderImage(FCoreStyle::Get().GetBrush(TEXT("BlackBrush")))
+			.OnMouseButtonDown(InArgs._OnMouseButtonDown)
+			.Padding(0)[InArgs._Content.Widget]);
+		
+	}
+
+	/**
+	* Set the handler to be invoked when the user presses a key.
+	*
+	* @param InHandler   Method to execute when the user presses a key
+	*/
+	void SetOnOnKeyDown(const FOnKeyDown& InHandler)
+	{
+		OnKeyDown = InHandler;
+	}	
+
+protected:
+	
+	FOnKeyDown OnKeyDown;	
+};
+
 TSharedPtr<FDefaultGameMoviePlayer> FDefaultGameMoviePlayer::MoviePlayer;
 
 TSharedPtr<FDefaultGameMoviePlayer> FDefaultGameMoviePlayer::Get()
@@ -83,10 +128,10 @@ void FDefaultGameMoviePlayer::Initialize()
 	TSharedRef<SWindow> GameWindow = UGameEngine::CreateGameWindow();
 
 	TSharedPtr<SViewport> MovieViewport;
-	LoadingScreenContents = SNew(SBorder)
-		.BorderImage( FCoreStyle::Get().GetBrush(TEXT("BlackBrush")) )
+
+	LoadingScreenContents = SNew(SDefaultMovieBorder)	
+		.OnKeyDown(this, &FDefaultGameMoviePlayer::OnLoadingScreenKeyDown)
 		.OnMouseButtonDown(this, &FDefaultGameMoviePlayer::OnLoadingScreenMouseButtonDown)
-		.Padding(0)
 		[
 			SNew(SOverlay)
 			+SOverlay::Slot()
@@ -234,6 +279,7 @@ void FDefaultGameMoviePlayer::WaitForMovieToFinish()
 		
 		const bool bAutoCompleteWhenLoadingCompletes = LoadingScreenAttributes.bAutoCompleteWhenLoadingCompletes;
 		const bool bWaitForManualStop = LoadingScreenAttributes.bWaitForManualStop;
+		bUserCalledFinish = true;
 
 		FSlateApplication& SlateApp = FSlateApplication::Get();
 		// Continue to wait until the user calls finish (if enabled) or when loading completes or the minimum enforced time (if any) has been reached.
@@ -451,6 +497,16 @@ EVisibility FDefaultGameMoviePlayer::GetViewportVisibility() const
 }
 
 FReply FDefaultGameMoviePlayer::OnLoadingScreenMouseButtonDown(const FGeometry& Geometry, const FPointerEvent& PointerEvent)
+{
+	return OnAnyDown();
+}
+
+FReply FDefaultGameMoviePlayer::OnLoadingScreenKeyDown(const FGeometry& Geometry, const FKeyEvent& KeyEvent)
+{
+	return OnAnyDown();
+}
+
+FReply FDefaultGameMoviePlayer::OnAnyDown()
 {
 	if (IsLoadingFinished())
 	{
@@ -472,7 +528,6 @@ FReply FDefaultGameMoviePlayer::OnLoadingScreenMouseButtonDown(const FGeometry&
 	return FReply::Handled();
 }
 
-
 void FDefaultGameMoviePlayer::OnPreLoadMap()
 {
 	FCoreUObjectDelegates::PostLoadMap.RemoveAll(this);
diff --git a/Engine/Source/Runtime/MoviePlayer/Private/DefaultGameMoviePlayer.h b/Engine/Source/Runtime/MoviePlayer/Private/DefaultGameMoviePlayer.h
--- a/Engine/Source/Runtime/MoviePlayer/Private/DefaultGameMoviePlayer.h
+++ b/Engine/Source/Runtime/MoviePlayer/Private/DefaultGameMoviePlayer.h
@@ -33,7 +33,10 @@ public:
 	virtual TStatId GetStatId() const override;
 	virtual bool IsTickable() const override;
 	
-	
+	/** Callback for clicking on the viewport */
+	FReply OnLoadingScreenMouseButtonDown(const FGeometry& Geometry, const FPointerEvent& PointerEvent);
+	FReply OnLoadingScreenKeyDown(const FGeometry& Geometry, const FKeyEvent& KeyEvent);
+
 private:
 
 	/** Ticks the underlying MovieStreamer.  Must be done exactly once before each DrawWindows call. */
@@ -50,10 +53,7 @@ private:
 	FOptionalSize GetMovieWidth() const;
 	FOptionalSize GetMovieHeight() const;
 	EVisibility GetSlateBackgroundVisibility() const;
-	EVisibility GetViewportVisibility() const;
-
-	/** Callback for clicking on the viewport */
-	FReply OnLoadingScreenMouseButtonDown(const FGeometry& Geometry, const FPointerEvent& PointerEvent);
+	EVisibility GetViewportVisibility() const;	
 	
 	/** Called via a delegate in the engine when maps start to load */
 	void OnPreLoadMap();
@@ -62,6 +62,8 @@ private:
 	void OnPostLoadMap();
 private:
 	FDefaultGameMoviePlayer();
+
+	FReply OnAnyDown();
 	
 	/** The movie streaming system that will be used by us */
 	TSharedPtr<IMovieStreamer> MovieStreamer;
