Copying //UE4/UT-Staging to //UE4/Dev-Main (Source CL 3221884)

==========================
MAJOR FEATURES + CHANGES
==========================

Change 3191861 by [REDACTED_PERSON]

Fix missing virtuals on ghost, fly, walk cheats

Change 3184553 by [REDACTED_PERSON]

FOnCurrentMovieClipFinished delegate for Mac movie player

Change 3183218 by [REDACTED_PERSON]

FOnCurrentMovieClipFinished delegate for movie player interface

Change 3175731 by [REDACTED_PERSON]

UWorld::ServerTravel crash (probably already in main by [REDACTED_PERSON])

Change 3172375 by [REDACTED_PERSON]

Allow windowed resolutions that match the desktop resolution

Change 3156755 by [REDACTED_PERSON]

ShouldSaveCheckpoint system for variable demo checkpoint timing

diff --git a/Engine/Plugins/Runtime/AndroidMoviePlayer/Source/AndroidMoviePlayer/Private/AndroidMovieStreamer.h b/Engine/Plugins/Runtime/AndroidMoviePlayer/Source/AndroidMoviePlayer/Private/AndroidMovieStreamer.h
--- a/Engine/Plugins/Runtime/AndroidMoviePlayer/Source/AndroidMoviePlayer/Private/AndroidMovieStreamer.h
+++ b/Engine/Plugins/Runtime/AndroidMoviePlayer/Source/AndroidMoviePlayer/Private/AndroidMovieStreamer.h
@@ -24,6 +24,8 @@ public:
 	virtual FString GetMovieName() override;
 	virtual bool IsLastMovieInPlaylist() override;
 
+	FOnCurrentMovieClipFinished OnCurrentMovieClipFinishedDelegate;
+	virtual FOnCurrentMovieClipFinished& OnCurrentMovieClipFinished() override { return OnCurrentMovieClipFinishedDelegate; }
 
 private:
 	/**
diff --git a/Engine/Plugins/Runtime/AppleMoviePlayer/Source/AppleMoviePlayer/Private/AppleMovieStreamer.h b/Engine/Plugins/Runtime/AppleMoviePlayer/Source/AppleMoviePlayer/Private/AppleMovieStreamer.h
--- a/Engine/Plugins/Runtime/AppleMoviePlayer/Source/AppleMoviePlayer/Private/AppleMovieStreamer.h
+++ b/Engine/Plugins/Runtime/AppleMoviePlayer/Source/AppleMoviePlayer/Private/AppleMovieStreamer.h
@@ -28,6 +28,8 @@ public:
 	virtual FString GetMovieName() override;
 	virtual bool IsLastMovieInPlaylist() override;
 
+	FOnCurrentMovieClipFinished OnCurrentMovieClipFinishedDelegate;
+	virtual FOnCurrentMovieClipFinished& OnCurrentMovieClipFinished() override { return OnCurrentMovieClipFinishedDelegate; }
 
 private:
 
diff --git a/Engine/Plugins/Runtime/WindowsMoviePlayer/Source/WindowsMoviePlayer/Private/WindowsMovieStreamer.cpp b/Engine/Plugins/Runtime/WindowsMoviePlayer/Source/WindowsMoviePlayer/Private/WindowsMovieStreamer.cpp
--- a/Engine/Plugins/Runtime/WindowsMoviePlayer/Source/WindowsMoviePlayer/Private/WindowsMovieStreamer.cpp
+++ b/Engine/Plugins/Runtime/WindowsMoviePlayer/Source/WindowsMoviePlayer/Private/WindowsMovieStreamer.cpp
@@ -170,6 +170,8 @@ void FMediaFoundationMovieStreamer::OpenNextMovie()
 
 void FMediaFoundationMovieStreamer::CloseMovie()
 {
+	BroadcastCurrentMovieClipFinished(GetMovieName());
+
 	if (Texture.IsValid())
 	{
 		TextureFreeList.Add(Texture);
diff --git a/Engine/Plugins/Runtime/WindowsMoviePlayer/Source/WindowsMoviePlayer/Private/WindowsMovieStreamer.h b/Engine/Plugins/Runtime/WindowsMoviePlayer/Source/WindowsMoviePlayer/Private/WindowsMovieStreamer.h
--- a/Engine/Plugins/Runtime/WindowsMoviePlayer/Source/WindowsMoviePlayer/Private/WindowsMovieStreamer.h
+++ b/Engine/Plugins/Runtime/WindowsMoviePlayer/Source/WindowsMoviePlayer/Private/WindowsMovieStreamer.h
@@ -37,6 +37,9 @@ public:
 	virtual FString GetMovieName() override;
 	virtual bool IsLastMovieInPlaylist() override;
 
+	FOnCurrentMovieClipFinished OnCurrentMovieClipFinishedDelegate;
+	virtual FOnCurrentMovieClipFinished& OnCurrentMovieClipFinished() override { return OnCurrentMovieClipFinishedDelegate; }
+
 	virtual FTexture2DRHIRef GetTexture() override { return Texture.IsValid() ? Texture->GetRHIRef() : nullptr; }
 
 private:
diff --git a/Engine/Source/Runtime/Engine/Classes/Engine/DemoNetDriver.h b/Engine/Source/Runtime/Engine/Classes/Engine/DemoNetDriver.h
--- a/Engine/Source/Runtime/Engine/Classes/Engine/DemoNetDriver.h
+++ b/Engine/Source/Runtime/Engine/Classes/Engine/DemoNetDriver.h
@@ -281,6 +281,8 @@ class ENGINE_API UDemoNetDriver : public UNetDriver
 
 	void		RespawnNecessaryNetStartupActors();
 
+	virtual bool ShouldSaveCheckpoint();
+
 	void		SaveCheckpoint();
 	void		TickCheckpoint();
 	bool		LoadCheckpoint( FArchive* GotoCheckpointArchive, int64 GotoCheckpointSkipExtraTimeInMS );
diff --git a/Engine/Source/Runtime/Engine/Classes/GameFramework/Character.h b/Engine/Source/Runtime/Engine/Classes/GameFramework/Character.h
--- a/Engine/Source/Runtime/Engine/Classes/GameFramework/Character.h
+++ b/Engine/Source/Runtime/Engine/Classes/GameFramework/Character.h
@@ -779,15 +779,15 @@ public:
 
 	UFUNCTION(reliable, client)
 	void ClientCheatWalk();
-	void ClientCheatWalk_Implementation();
+	virtual void ClientCheatWalk_Implementation();
 
 	UFUNCTION(reliable, client)
 	void ClientCheatFly();
-	void ClientCheatFly_Implementation();
+	virtual void ClientCheatFly_Implementation();
 
 	UFUNCTION(reliable, client)
 	void ClientCheatGhost();
-	void ClientCheatGhost_Implementation();
+	virtual void ClientCheatGhost_Implementation();
 
 	// Root Motion
 public:
diff --git a/Engine/Source/Runtime/Engine/Private/DemoNetDriver.cpp b/Engine/Source/Runtime/Engine/Private/DemoNetDriver.cpp
--- a/Engine/Source/Runtime/Engine/Private/DemoNetDriver.cpp
+++ b/Engine/Source/Runtime/Engine/Private/DemoNetDriver.cpp
@@ -1671,9 +1671,7 @@ void UDemoNetDriver::TickDemoRecord( float DeltaSeconds )
 	{
 		check( !PendingCheckpointActors.Num() );		// We early out above, so this shouldn't be possible
 
-		const double CHECKPOINT_DELAY = CVarCheckpointUploadDelayInSeconds.GetValueOnAnyThread();
-
-		if ( DemoCurrentTime - LastCheckpointTime > CHECKPOINT_DELAY )
+		if (ShouldSaveCheckpoint())
 		{
 			SaveCheckpoint();
 			LastCheckpointTime = DemoCurrentTime;
@@ -1681,6 +1679,18 @@ void UDemoNetDriver::TickDemoRecord( float DeltaSeconds )
 	}
 }
 
+bool UDemoNetDriver::ShouldSaveCheckpoint()
+{
+	const double CHECKPOINT_DELAY = CVarCheckpointUploadDelayInSeconds.GetValueOnGameThread();
+
+	if (DemoCurrentTime - LastCheckpointTime > CHECKPOINT_DELAY)
+	{
+		return true;
+	}
+
+	return false;
+}
+
 void UDemoNetDriver::PauseChannels( const bool bPause )
 {
 	if ( bPause == bChannelsArePaused )
diff --git a/Engine/Source/Runtime/Engine/Private/GameEngine.cpp b/Engine/Source/Runtime/Engine/Private/GameEngine.cpp
--- a/Engine/Source/Runtime/Engine/Private/GameEngine.cpp
+++ b/Engine/Source/Runtime/Engine/Private/GameEngine.cpp
@@ -252,8 +252,8 @@ void UGameEngine::ConditionallyOverrideSettings(int32& ResolutionX, int32& Resol
 	// Optionally force the resolution by passing -ForceRes
 	const bool bForceRes = FParse::Param(FCommandLine::Get(), TEXT("ForceRes"));
 
-	//Dont allow a resolution bigger then the desktop found a convenient one
-	if (!bForceRes && !IsRunningDedicatedServer() && ((ResolutionX <= 0 || ResolutionX >= MaxResolutionX) || (ResolutionY <= 0 || ResolutionY >= MaxResolutionY)))
+	//Don't allow a resolution bigger then the desktop found a convenient one
+	if (!bForceRes && !IsRunningDedicatedServer() && ((ResolutionX <= 0 || ResolutionX > MaxResolutionX) || (ResolutionY <= 0 || ResolutionY > MaxResolutionY)))
 	{
 		ResolutionX = MaxResolutionX;
 		ResolutionY = MaxResolutionY;
diff --git a/Engine/Source/Runtime/MoviePlayer/Private/DefaultGameMoviePlayer.cpp b/Engine/Source/Runtime/MoviePlayer/Private/DefaultGameMoviePlayer.cpp
--- a/Engine/Source/Runtime/MoviePlayer/Private/DefaultGameMoviePlayer.cpp
+++ b/Engine/Source/Runtime/MoviePlayer/Private/DefaultGameMoviePlayer.cpp
@@ -133,6 +133,7 @@ FDefaultGameMoviePlayer::~FDefaultGameMoviePlayer()
 void FDefaultGameMoviePlayer::RegisterMovieStreamer(TSharedPtr<IMovieStreamer> InMovieStreamer)
 {
 	MovieStreamer = InMovieStreamer;
+	MovieStreamer->OnCurrentMovieClipFinished().AddRaw(this, &FDefaultGameMoviePlayer::BroadcastMovieClipFinished);
 }
 
 void FDefaultGameMoviePlayer::SetSlateRenderer(TSharedPtr<FSlateRenderer> InSlateRenderer)
diff --git a/Engine/Source/Runtime/MoviePlayer/Private/DefaultGameMoviePlayer.h b/Engine/Source/Runtime/MoviePlayer/Private/DefaultGameMoviePlayer.h
--- a/Engine/Source/Runtime/MoviePlayer/Private/DefaultGameMoviePlayer.h
+++ b/Engine/Source/Runtime/MoviePlayer/Private/DefaultGameMoviePlayer.h
@@ -40,6 +40,7 @@ public:
 
 	virtual FOnPrepareLoadingScreen& OnPrepareLoadingScreen() override { return OnPrepareLoadingScreenDelegate; }
 	virtual FOnMoviePlaybackFinished& OnMoviePlaybackFinished() override { return OnMoviePlaybackFinishedDelegate; }
+	virtual FOnMovieClipFinished& OnMovieClipFinished() override { return OnMovieClipFinishedDelegate; }
 
 	/** FTickableObjectRenderThread interface */
 	virtual void Tick( float DeltaTime ) override;
@@ -118,6 +119,9 @@ private:
 	
 	FOnMoviePlaybackFinished OnMoviePlaybackFinishedDelegate;
 
+	FOnMovieClipFinished OnMovieClipFinishedDelegate;
+
+
 	/** The last time a movie was started */
 	double LastPlayTime;
 
diff --git a/Engine/Source/Runtime/MoviePlayer/Private/NullMoviePlayer.h b/Engine/Source/Runtime/MoviePlayer/Private/NullMoviePlayer.h
--- a/Engine/Source/Runtime/MoviePlayer/Private/NullMoviePlayer.h
+++ b/Engine/Source/Runtime/MoviePlayer/Private/NullMoviePlayer.h
@@ -35,6 +35,7 @@ public:
 	virtual void SetupLoadingScreenFromIni() override {}
 	virtual FOnPrepareLoadingScreen& OnPrepareLoadingScreen() override { return OnPrepareLoadingScreenDelegate; }
 	virtual FOnMoviePlaybackFinished& OnMoviePlaybackFinished() override { return OnMoviePlaybackFinishedDelegate; }
+	virtual FOnMovieClipFinished& OnMovieClipFinished() override { return OnMovieClipFinishedDelegate; }
 	virtual void SetSlateOverlayWidget(TSharedPtr<SWidget> NewOverlayWidget) override { }
 	virtual bool WillAutoCompleteWhenLoadFinishes() override { return false; }
 	virtual FString GetMovieName() override { return TEXT(""); }
@@ -52,4 +53,6 @@ private:
 	FOnPrepareLoadingScreen OnPrepareLoadingScreenDelegate;
 
 	FOnMoviePlaybackFinished OnMoviePlaybackFinishedDelegate;
+	FOnMovieClipFinished OnMovieClipFinishedDelegate;
+
 };
diff --git a/Engine/Source/Runtime/MoviePlayer/Public/MoviePlayer.h b/Engine/Source/Runtime/MoviePlayer/Public/MoviePlayer.h
--- a/Engine/Source/Runtime/MoviePlayer/Public/MoviePlayer.h
+++ b/Engine/Source/Runtime/MoviePlayer/Public/MoviePlayer.h
@@ -86,6 +86,13 @@ public:
 	virtual FTexture2DRHIRef GetTexture() { return nullptr; }
 
 	virtual ~IMovieStreamer() {}
+
+	DECLARE_EVENT_OneParam(IMovieStreamer, FOnCurrentMovieClipFinished, const FString&)
+	virtual FOnCurrentMovieClipFinished& OnCurrentMovieClipFinished() = 0;
+
+	void BroadcastCurrentMovieClipFinished(const FString& MovieClipThatFinished) { OnCurrentMovieClipFinished().Broadcast(MovieClipThatFinished); }
+
+
 };
 
 
@@ -189,11 +196,15 @@ public:
 
 	DECLARE_EVENT(IGameMoviePlayer, FOnMoviePlaybackFinished)
 	virtual FOnMoviePlaybackFinished& OnMoviePlaybackFinished() = 0;
-	
+
+	DECLARE_EVENT_OneParam(IGameMoviePlayer, FOnMovieClipFinished, const FString&)
+	virtual FOnMovieClipFinished& OnMovieClipFinished() = 0;
+
 	/** Allows for a slate overlay widget to be set after playback. */
 	virtual void SetSlateOverlayWidget(TSharedPtr<SWidget> NewOverlayWidget) = 0;
 
 	void BroadcastMoviePlaybackFinished() { OnMoviePlaybackFinished().Broadcast(); }
+	void BroadcastMovieClipFinished(const FString& MovieClipThatFinished) { OnMovieClipFinished().Broadcast(MovieClipThatFinished); }
 
 	/** This function shouild return true if the movie will auto-complete the sequence when background loading has finished */
 	virtual bool WillAutoCompleteWhenLoadFinishes() = 0;
