Remove ensure in GetMoviePlayer()

MattK is going to be responsible for the thread safety of the slate loading thread / movie player interaction

diff --git a/Engine/Source/Runtime/MoviePlayer/Private/MoviePlayer.cpp b/Engine/Source/Runtime/MoviePlayer/Private/MoviePlayer.cpp
--- a/Engine/Source/Runtime/MoviePlayer/Private/MoviePlayer.cpp
+++ b/Engine/Source/Runtime/MoviePlayer/Private/MoviePlayer.cpp
@@ -68,8 +68,6 @@ TSharedRef<SWidget> FLoadingScreenAttributes::NewTestLoadingScreenWidget()
 
 TSharedPtr<IGameMoviePlayer> GetMoviePlayer()
 {
-	ensureMsgf(!IsInSlateThread(), TEXT("GetMoviePlayer() is not safe to be called from the Slate Thread"));
-
 	if (!IsMoviePlayerEnabled() || GUsingNullRHI)
 	{
 		return FNullGameMoviePlayer::Get();
