Fix crash on shutdown when the engine initializes without creating a movie player (launching without a project, launching with module errors)

diff --git a/Engine/Source/Runtime/MoviePlayer/Private/MoviePlayer.cpp b/Engine/Source/Runtime/MoviePlayer/Private/MoviePlayer.cpp
--- a/Engine/Source/Runtime/MoviePlayer/Private/MoviePlayer.cpp
+++ b/Engine/Source/Runtime/MoviePlayer/Private/MoviePlayer.cpp
@@ -84,15 +84,19 @@ void CreateMoviePlayer()
 
 void DestroyMoviePlayer()
 {
-	GetMoviePlayer()->Shutdown();
+	IGameMoviePlayer* MoviePlayer = GetMoviePlayer();
 
-	if (!IsMoviePlayerEnabled() || GUsingNullRHI)
-	{
-		return FNullGameMoviePlayer::Destroy();
-	}
-	else
+	if(MoviePlayer)
 	{
-		return FDefaultGameMoviePlayer::Destroy();
+
+		if (!IsMoviePlayerEnabled() || GUsingNullRHI)
+		{
+			return FNullGameMoviePlayer::Destroy();
+		}
+		else
+		{
+			return FDefaultGameMoviePlayer::Destroy();
+		}
 	}
 }
 
