Fixed startup movies appearing black for packaged games.  This is due to movie players being registered after the movie viewport has been created.  We now set the movie streamer on the viewport whenever it is registered

diff --git a/Engine/Source/Runtime/MoviePlayer/Private/DefaultGameMoviePlayer.cpp b/Engine/Source/Runtime/MoviePlayer/Private/DefaultGameMoviePlayer.cpp
--- a/Engine/Source/Runtime/MoviePlayer/Private/DefaultGameMoviePlayer.cpp
+++ b/Engine/Source/Runtime/MoviePlayer/Private/DefaultGameMoviePlayer.cpp
@@ -140,6 +140,13 @@ void FDefaultGameMoviePlayer::RegisterMovieStreamer(TSharedPtr<IMovieStreamer> I
 {
 	MovieStreamer = InMovieStreamer;
 	MovieStreamer->OnCurrentMovieClipFinished().AddRaw(this, &FDefaultGameMoviePlayer::BroadcastMovieClipFinished);
+
+	// Always set the viewport interface when the movie streamer is registered.  This fixes cases movies not appearing due to module load order where the movie player is initialized after a movie streamer is registered.
+	// It also fixes cases of a different movie streamer being registered after initialization for any reason
+	if (MovieStreamer.IsValid() && MovieViewportWeakPtr.IsValid())
+	{
+		MovieViewportWeakPtr.Pin()->SetViewportInterface(MovieStreamer->GetViewportInterface().ToSharedRef());
+	}
 }
 
 void FDefaultGameMoviePlayer::Initialize(FSlateRenderer& InSlateRenderer)
@@ -213,6 +220,8 @@ void FDefaultGameMoviePlayer::Initialize(FSlateRenderer& InSlateRenderer)
 			]
 		];
 
+	MovieViewportWeakPtr = MovieViewport;
+
 	if (MovieStreamer.IsValid())
 	{
 		MovieViewport->SetViewportInterface( MovieStreamer->GetViewportInterface().ToSharedRef() );
diff --git a/Engine/Source/Runtime/MoviePlayer/Private/DefaultGameMoviePlayer.h b/Engine/Source/Runtime/MoviePlayer/Private/DefaultGameMoviePlayer.h
--- a/Engine/Source/Runtime/MoviePlayer/Private/DefaultGameMoviePlayer.h
+++ b/Engine/Source/Runtime/MoviePlayer/Private/DefaultGameMoviePlayer.h
@@ -149,6 +149,8 @@ private:
 	TSharedPtr<class SBorder> UserWidgetHolder;
 	/** Virtual window that we render to instead of the main slate window (for thread safety).  Shares only the same backbuffer as the main window */
 	TSharedPtr<class SVirtualWindow> VirtualRenderWindow;
+	/** Viewport responsible for displaying the movie player render target */
+	TWeakPtr<class SViewport> MovieViewportWeakPtr;
 
 	/** The threading mechanism with which we handle running slate on another thread */
 	class FSlateLoadingSynchronizationMechanism* SyncMechanism;
