Allow multiple movie streamer plugins to register themselves with DefaultGameMoviePlayer and pick the right one for the starup movie playback based on the format of the file.

diff --git a/Engine/Source/Runtime/MoviePlayer/Private/DefaultGameMoviePlayer.cpp b/Engine/Source/Runtime/MoviePlayer/Private/DefaultGameMoviePlayer.cpp
--- a/Engine/Source/Runtime/MoviePlayer/Private/DefaultGameMoviePlayer.cpp
+++ b/Engine/Source/Runtime/MoviePlayer/Private/DefaultGameMoviePlayer.cpp
@@ -142,14 +142,10 @@ FDefaultGameMoviePlayer::~FDefaultGameMoviePlayer()
 
 void FDefaultGameMoviePlayer::RegisterMovieStreamer(TSharedPtr<IMovieStreamer> InMovieStreamer)
 {
-	MovieStreamer = InMovieStreamer;
-	MovieStreamer->OnCurrentMovieClipFinished().AddRaw(this, &FDefaultGameMoviePlayer::BroadcastMovieClipFinished);
-
-	// Always set the viewport interface when the movie streamer is registered.  This fixes cases movies not appearing due to module load order where the movie player is initialized after a movie streamer is registered.
-	// It also fixes cases of a different movie streamer being registered after initialization for any reason
-	if (MovieStreamer.IsValid() && MovieViewportWeakPtr.IsValid())
+	if (InMovieStreamer.IsValid() && !MovieStreamers.Contains(InMovieStreamer))
 	{
-		MovieViewportWeakPtr.Pin()->SetViewportInterface(MovieStreamer->GetViewportInterface().ToSharedRef());
+		MovieStreamers.Add(InMovieStreamer);
+		InMovieStreamer->OnCurrentMovieClipFinished().AddRaw(this, &FDefaultGameMoviePlayer::BroadcastMovieClipFinished);
 	}
 }
 
@@ -228,12 +224,6 @@ void FDefaultGameMoviePlayer::Initialize(FSlateRenderer& InSlateRenderer, TShare
 		];
 
 	MovieViewportWeakPtr = MovieViewport;
-
-	if (MovieStreamer.IsValid())
-	{
-		MovieViewport->SetViewportInterface( MovieStreamer->GetViewportInterface().ToSharedRef() );
-	}
-	
 	MovieViewport->SetActive(true);
 
 	// Register the movie viewport so that it can receive user input.
@@ -284,7 +274,8 @@ void FDefaultGameMoviePlayer::Shutdown()
 	MainWindow.Reset();
 	VirtualRenderWindow.Reset();
 
-	MovieStreamer.Reset();
+	MovieStreamers.Empty();
+	ActiveMovieStreamer.Reset();
 
 	LoadingScreenAttributes = FLoadingScreenAttributes();
 
@@ -359,13 +350,24 @@ bool FDefaultGameMoviePlayer::PlayMovie()
 		
 		LastPlayTime = FPlatformTime::Seconds();
 
-		bool bIsInitialized = true;
+		ActiveMovieStreamer.Reset();
 		if (MovieStreamingIsPrepared())
 		{
-			bIsInitialized = MovieStreamer->Init(LoadingScreenAttributes.MoviePaths, LoadingScreenAttributes.PlaybackType);
+			for (TSharedPtr<IMovieStreamer> MovieStreamer : MovieStreamers)
+			{
+				if (MovieStreamer->Init(LoadingScreenAttributes.MoviePaths, LoadingScreenAttributes.PlaybackType))
+				{
+					ActiveMovieStreamer = MovieStreamer;
+					if (MovieViewportWeakPtr.IsValid())
+					{
+						MovieViewportWeakPtr.Pin()->SetViewportInterface(MovieStreamer->GetViewportInterface().ToSharedRef());
+					}
+					break;
+				}
+			}
 		}
 
-		if (bIsInitialized)
+		if (ActiveMovieStreamer.IsValid())
 		{
 			MovieStreamingIsDone.Set(MovieStreamingIsPrepared() ? 0 : 1);
 			LoadingIsDone.Set(0);
@@ -464,7 +466,7 @@ void FDefaultGameMoviePlayer::WaitForMovieToFinish(bool bAllowEngineTick)
 			||	(bEnforceMinimumTime && (FPlatformTime::Seconds() - LastPlayTime) < LoadingScreenAttributes.MinimumLoadingScreenDisplayTime)))
 		{
 			// If we are in a loading loop, and this is the last movie in the playlist.. assume you can break out.
-			if (MovieStreamer.IsValid() && LoadingScreenAttributes.PlaybackType == MT_LoadingLoop && MovieStreamer->IsLastMovieInPlaylist())
+			if (ActiveMovieStreamer.IsValid() && LoadingScreenAttributes.PlaybackType == MT_LoadingLoop && ActiveMovieStreamer->IsLastMovieInPlaylist())
 			{
 				break;
 			}
@@ -474,9 +476,9 @@ void FDefaultGameMoviePlayer::WaitForMovieToFinish(bool bAllowEngineTick)
 				// Break out of the loop if the main window is closed during the movie.
 				if ( !MainWindow.IsValid() || bMainWindowClosed.Load() )
 				{
-					if (MovieStreamer.IsValid())
+					if (ActiveMovieStreamer.IsValid())
 					{
-						MovieStreamer->ForceCompletion();
+						ActiveMovieStreamer->ForceCompletion();
 					}
 					break;
 				}
@@ -534,15 +536,15 @@ void FDefaultGameMoviePlayer::WaitForMovieToFinish(bool bAllowEngineTick)
 
 		FlushRenderingCommands();
 
-		if( MovieStreamer.IsValid() )
+		if( ActiveMovieStreamer.IsValid() )
 		{
-			MovieStreamer->ForceCompletion();
+			ActiveMovieStreamer->ForceCompletion();
 		}
 
 		// Allow the movie streamer to clean up any resources it uses once there are no movies to play.
-		if( MovieStreamer.IsValid() )
+		if( ActiveMovieStreamer.IsValid() )
 		{
-			MovieStreamer->Cleanup();
+			ActiveMovieStreamer->Cleanup();
 		}
 	
 		// Finally, clear out the loading screen attributes, forcing users to always
@@ -602,9 +604,9 @@ void FDefaultGameMoviePlayer::Tick( float DeltaTime )
 
 void FDefaultGameMoviePlayer::TickStreamer(float DeltaTime)
 {	
-	if (MovieStreamingIsPrepared() && !IsMovieStreamingFinished())
+	if (MovieStreamingIsPrepared() && ActiveMovieStreamer.IsValid() && !IsMovieStreamingFinished())
 	{
-		const bool bMovieIsDone = MovieStreamer->Tick(DeltaTime);		
+		const bool bMovieIsDone = ActiveMovieStreamer->Tick(DeltaTime);
 		if (bMovieIsDone)
 		{
 			MovieStreamingIsDone.Set(1);
@@ -613,7 +615,7 @@ void FDefaultGameMoviePlayer::TickStreamer(float DeltaTime)
 		IStereoLayers* StereoLayers;
 		if (GEngine && GEngine->StereoRenderingDevice.IsValid() && (StereoLayers = GEngine->StereoRenderingDevice->GetStereoLayers()) != nullptr)
 		{
-			FTexture2DRHIRef Movie2DTexture = MovieStreamer->GetTexture();
+			FTexture2DRHIRef Movie2DTexture = ActiveMovieStreamer->GetTexture();
 			FTextureRHIRef MovieTexture;
 			if (Movie2DTexture.IsValid() && !bMovieIsDone)
 			{
@@ -697,15 +699,15 @@ void FDefaultGameMoviePlayer::SetupLoadingScreenFromIni()
 
 bool FDefaultGameMoviePlayer::MovieStreamingIsPrepared() const
 {
-	return MovieStreamer.IsValid() && LoadingScreenAttributes.MoviePaths.Num() > 0;
+	return MovieStreamers.Num() > 0 && LoadingScreenAttributes.MoviePaths.Num() > 0;
 }
 
 FVector2D FDefaultGameMoviePlayer::GetMovieSize() const
 {
 	const FVector2D ScreenSize = MainWindow.Pin()->GetClientSizeInScreen();
-	if (MovieStreamingIsPrepared())
+	if (MovieStreamingIsPrepared() && ActiveMovieStreamer.IsValid())
 	{
-		const float MovieAspectRatio = MovieStreamer->GetAspectRatio();
+		const float MovieAspectRatio = ActiveMovieStreamer->GetAspectRatio();
 		const float ScreenAspectRatio = ScreenSize.X / ScreenSize.Y;
 		if (MovieAspectRatio < ScreenAspectRatio)
 		{
@@ -733,12 +735,12 @@ FOptionalSize FDefaultGameMoviePlayer::GetMovieHeight() const
 
 EVisibility FDefaultGameMoviePlayer::GetSlateBackgroundVisibility() const
 {
-	return MovieStreamingIsPrepared() && !IsMovieStreamingFinished() ? EVisibility::Collapsed : EVisibility::Visible;
+	return MovieStreamingIsPrepared() && ActiveMovieStreamer.IsValid() && !IsMovieStreamingFinished() ? EVisibility::Collapsed : EVisibility::Visible;
 }
 
 EVisibility FDefaultGameMoviePlayer::GetViewportVisibility() const
 {
-	return MovieStreamingIsPrepared() && !IsMovieStreamingFinished() ? EVisibility::Visible : EVisibility::Collapsed;
+	return MovieStreamingIsPrepared() && ActiveMovieStreamer.IsValid() && !IsMovieStreamingFinished() ? EVisibility::Visible : EVisibility::Collapsed;
 }
 
 FReply FDefaultGameMoviePlayer::OnLoadingScreenMouseButtonDown(const FGeometry& Geometry, const FPointerEvent& PointerEvent)
@@ -758,9 +760,9 @@ FReply FDefaultGameMoviePlayer::OnAnyDown()
 		if (LoadingScreenAttributes.bMoviesAreSkippable)
 		{
 			MovieStreamingIsDone.Set(1);
-			if (MovieStreamer.IsValid())
+			if (ActiveMovieStreamer.IsValid())
 			{
-				MovieStreamer->ForceCompletion();
+				ActiveMovieStreamer->ForceCompletion();
 			}
 		}
 
@@ -794,7 +796,7 @@ void FDefaultGameMoviePlayer::OnPostLoadMap(UWorld* LoadedWorld)
 
 void FDefaultGameMoviePlayer::SetSlateOverlayWidget(TSharedPtr<SWidget> NewOverlayWidget)
 {
-	if (MovieStreamer.IsValid() && UserWidgetHolder.IsValid())
+	if (ActiveMovieStreamer.IsValid() && UserWidgetHolder.IsValid())
 	{
 		UserWidgetHolder->SetContent(NewOverlayWidget.ToSharedRef());
 	}
@@ -802,17 +804,17 @@ void FDefaultGameMoviePlayer::SetSlateOverlayWidget(TSharedPtr<SWidget> NewOverl
 
 bool FDefaultGameMoviePlayer::WillAutoCompleteWhenLoadFinishes()
 {
-	return LoadingScreenAttributes.bAutoCompleteWhenLoadingCompletes || (LoadingScreenAttributes.PlaybackType == MT_LoadingLoop && (MovieStreamer.IsValid() && MovieStreamer->IsLastMovieInPlaylist()));
+	return LoadingScreenAttributes.bAutoCompleteWhenLoadingCompletes || (LoadingScreenAttributes.PlaybackType == MT_LoadingLoop && (ActiveMovieStreamer.IsValid() && ActiveMovieStreamer->IsLastMovieInPlaylist()));
 }
 
 FString FDefaultGameMoviePlayer::GetMovieName()
 {
-	return MovieStreamer.IsValid() ? MovieStreamer->GetMovieName() : TEXT("");
+	return ActiveMovieStreamer.IsValid() ? ActiveMovieStreamer->GetMovieName() : TEXT("");
 }
 
 bool FDefaultGameMoviePlayer::IsLastMovieInPlaylist()
 {
-	return MovieStreamer.IsValid() ? MovieStreamer->IsLastMovieInPlaylist() : false;
+	return ActiveMovieStreamer.IsValid() ? ActiveMovieStreamer->IsLastMovieInPlaylist() : false;
 }
 
 FMoviePlayerWidgetRenderer::FMoviePlayerWidgetRenderer(TSharedPtr<SWindow> InMainWindow, TSharedPtr<SVirtualWindow> InVirtualRenderWindow, FSlateRenderer* InRenderer)
@@ -879,25 +881,25 @@ void FDefaultGameMoviePlayer::ForceCompletion()
 	bUserCalledFinish = true;
 	MovieStreamingIsDone.Set(1);
 
-	if (MovieStreamer.IsValid())
+	if (ActiveMovieStreamer.IsValid())
 	{
-		MovieStreamer->ForceCompletion();
+		ActiveMovieStreamer->ForceCompletion();
 	}
 }
 /*Interrupts*/
 void FDefaultGameMoviePlayer::Suspend()
 {
-	if (MovieStreamer.IsValid())
+	if (ActiveMovieStreamer.IsValid())
 	{
-		MovieStreamer->Suspend();
+		ActiveMovieStreamer->Suspend();
 	}
 }
 
 void FDefaultGameMoviePlayer::Resume()
 {
-	if (MovieStreamer.IsValid())
+	if (ActiveMovieStreamer.IsValid())
 	{
-		MovieStreamer->Resume();
+		ActiveMovieStreamer->Resume();
 	}
 }
 
diff --git a/Engine/Source/Runtime/MoviePlayer/Private/DefaultGameMoviePlayer.h b/Engine/Source/Runtime/MoviePlayer/Private/DefaultGameMoviePlayer.h
--- a/Engine/Source/Runtime/MoviePlayer/Private/DefaultGameMoviePlayer.h
+++ b/Engine/Source/Runtime/MoviePlayer/Private/DefaultGameMoviePlayer.h
@@ -143,9 +143,10 @@ private:
 
 	FReply OnAnyDown();
 	
-	/** The movie streaming system that will be used by us */
-	TSharedPtr<IMovieStreamer> MovieStreamer;
-	
+	/** The movie streaming systems that will be used by us */
+	TArray<TSharedPtr<IMovieStreamer>> MovieStreamers;
+	TSharedPtr<IMovieStreamer> ActiveMovieStreamer;
+
 	/** The window that the loading screen resides in */
 	TWeakPtr<class SWindow> MainWindow;
 	/** The widget which includes all contents of the loading screen, widgets and movie player and all */
