Fixed DefaultGameMoviePlayer to allow widgets to be displayed even if no movies have been added

diff --git a/Engine/Source/Runtime/MoviePlayer/Private/DefaultGameMoviePlayer.cpp b/Engine/Source/Runtime/MoviePlayer/Private/DefaultGameMoviePlayer.cpp
--- a/Engine/Source/Runtime/MoviePlayer/Private/DefaultGameMoviePlayer.cpp
+++ b/Engine/Source/Runtime/MoviePlayer/Private/DefaultGameMoviePlayer.cpp
@@ -367,7 +367,7 @@ bool FDefaultGameMoviePlayer::PlayMovie()
 			}
 		}
 
-		if (ActiveMovieStreamer.IsValid())
+		if (ActiveMovieStreamer.IsValid() || !MovieStreamingIsPrepared())
 		{
 			MovieStreamingIsDone.Set(MovieStreamingIsPrepared() ? 0 : 1);
 			LoadingIsDone.Set(0);
