Fix startup movie audio continuing to play on Windows platform after application closed

diff --git a/Engine/Source/Runtime/MoviePlayer/Private/DefaultGameMoviePlayer.cpp b/Engine/Source/Runtime/MoviePlayer/Private/DefaultGameMoviePlayer.cpp
--- a/Engine/Source/Runtime/MoviePlayer/Private/DefaultGameMoviePlayer.cpp
+++ b/Engine/Source/Runtime/MoviePlayer/Private/DefaultGameMoviePlayer.cpp
@@ -111,6 +111,7 @@ FDefaultGameMoviePlayer::FDefaultGameMoviePlayer()
 	, LoadingScreenAttributes()
 	, LastPlayTime(0.0)
 	, bInitialized(false)
+	, bMainWindowClosed(false)
 {
 	FCoreDelegates::IsLoadingMovieCurrentlyPlaying.BindRaw(this, &FDefaultGameMoviePlayer::IsMovieCurrentlyPlaying);
     FCoreDelegates::RegisterMovieStreamerDelegate.AddRaw(this, &FDefaultGameMoviePlayer::RegisterMovieStreamer);
@@ -238,12 +239,26 @@ void FDefaultGameMoviePlayer::Initialize(FSlateRenderer& InSlateRenderer, TShare
 	}
 
 	MainWindow = GameWindow;
+
+	GameWindow->GetOnWindowClosedEvent().AddRaw(this, &FDefaultGameMoviePlayer::OnMainWindowClosed);
+}
+
+void FDefaultGameMoviePlayer::OnMainWindowClosed(const TSharedRef<SWindow>& Window)
+{
+	bMainWindowClosed = true;
 }
 
 void FDefaultGameMoviePlayer::Shutdown()
 {
 	UE_LOG(LogMoviePlayer, Log, TEXT("Shutting down movie player"));
 
+	TSharedPtr<SWindow> MainWindowShared = MainWindow.Pin();
+	if (MainWindowShared.IsValid())
+	{
+		MainWindowShared->GetOnWindowClosedEvent().RemoveAll(this);
+		MainWindowShared.Reset();
+	}
+
 	StopMovie();
 	WaitForMovieToFinish();
 
@@ -451,8 +466,12 @@ void FDefaultGameMoviePlayer::WaitForMovieToFinish(bool bAllowEngineTick)
 			if (FSlateApplication::IsInitialized())
 			{
 				// Break out of the loop if the main window is closed during the movie.
-				if ( !MainWindow.IsValid() )
+				if ( !MainWindow.IsValid() || bMainWindowClosed.Load() )
 				{
+					if (MovieStreamer.IsValid())
+					{
+						MovieStreamer->ForceCompletion();
+					}
 					break;
 				}
 
diff --git a/Engine/Source/Runtime/MoviePlayer/Private/DefaultGameMoviePlayer.h b/Engine/Source/Runtime/MoviePlayer/Private/DefaultGameMoviePlayer.h
--- a/Engine/Source/Runtime/MoviePlayer/Private/DefaultGameMoviePlayer.h
+++ b/Engine/Source/Runtime/MoviePlayer/Private/DefaultGameMoviePlayer.h
@@ -108,6 +108,8 @@ public:
 
 	float GetViewportDPIScale() const;
 
+	void OnMainWindowClosed(const TSharedRef<SWindow>& Window);
+
 private:
 
 	/** Ticks the underlying MovieStreamer.  Must be done exactly once before each DrawWindows call. */
@@ -168,6 +170,9 @@ private:
 	/** User has called finish (needed if LoadingScreenAttributes.bAutoCompleteWhenLoadingCompletes is off) */
 	bool bUserCalledFinish;
 	
+	/** Main window has closed, stop movie playback */
+	TAtomic<bool> bMainWindowClosed;
+
 	/** Attributes of the loading screen we are currently displaying */
 	FLoadingScreenAttributes LoadingScreenAttributes;
 
diff --git a/Engine/Source/Runtime/SlateCore/Private/Widgets/SWindow.cpp b/Engine/Source/Runtime/SlateCore/Private/Widgets/SWindow.cpp
--- a/Engine/Source/Runtime/SlateCore/Private/Widgets/SWindow.cpp
+++ b/Engine/Source/Runtime/SlateCore/Private/Widgets/SWindow.cpp
@@ -1271,10 +1271,11 @@ void SWindow::DestroyWindowImmediately()
 	}
 }
 
-/** Calls the OnWindowClosed delegate when this window is about to be closed */
+/** Calls OnWindowClosed delegate and WindowClosedEvent when this window is about to be closed */
 void SWindow::NotifyWindowBeingDestroyed()
 {
 	OnWindowClosed.ExecuteIfBound( SharedThis( this ) );
+	WindowClosedEvent.Broadcast( SharedThis( this ) );
 
 #if WITH_EDITOR
 	if(bIsModalWindow)
diff --git a/Engine/Source/Runtime/SlateCore/Public/Widgets/SWindow.h b/Engine/Source/Runtime/SlateCore/Public/Widgets/SWindow.h
--- a/Engine/Source/Runtime/SlateCore/Public/Widgets/SWindow.h
+++ b/Engine/Source/Runtime/SlateCore/Public/Widgets/SWindow.h
@@ -48,6 +48,7 @@ DECLARE_MULTICAST_DELEGATE( FOnWindowDeactivatedEvent );
 
 /** Notification that a window is about to be closed */
 DECLARE_DELEGATE_OneParam( FOnWindowClosed, const TSharedRef<SWindow>& );
+DECLARE_MULTICAST_DELEGATE_OneParam( FOnWindowClosedEvent, const TSharedRef<SWindow>& );
 
 /** Notification that a window has been moved */
 DECLARE_DELEGATE_OneParam( FOnWindowMoved, const TSharedRef<SWindow>& );
@@ -567,6 +568,9 @@ public:
 	/** Sets the delegate to execute right before the window is closed */
 	void SetOnWindowClosed( const FOnWindowClosed& InDelegate );
 
+	/** Gets the multicast delegate to execute right before the window is closed */
+	FOnWindowClosedEvent& GetOnWindowClosedEvent() { return WindowClosedEvent; }
+
 	/** Sets the delegate to execute right after the window has been moved */
 	void SetOnWindowMoved( const FOnWindowMoved& InDelegate);
 
@@ -1095,6 +1099,7 @@ protected:
 
 	/** Invoked when the window is about to be closed. */
 	FOnWindowClosed OnWindowClosed;
+	FOnWindowClosedEvent WindowClosedEvent;
 
 	/** Invoked when the window is moved */
 	FOnWindowMoved OnWindowMoved;
