Adjusted sleep time for message pumping in FSlateLoadingSynchronizationMechanism::DestroySlateThread

#rnx

diff --git a/Engine/Source/Runtime/MoviePlayer/Private/MoviePlayerThreading.cpp b/Engine/Source/Runtime/MoviePlayer/Private/MoviePlayerThreading.cpp
--- a/Engine/Source/Runtime/MoviePlayer/Private/MoviePlayerThreading.cpp
+++ b/Engine/Source/Runtime/MoviePlayer/Private/MoviePlayerThreading.cpp
@@ -75,7 +75,7 @@ void FSlateLoadingSynchronizationMechanism::DestroySlateThread()
 		{
 			FPlatformApplicationMisc::PumpMessages(false);
 
-			FPlatformProcess::Sleep(0.1f);
+			FPlatformProcess::Sleep(0.016f);
 		}
 
 		delete SlateLoadingThread;
