UE-67123 //UE4/Dev-Mobile - Compile / Incremental UE4Editor - FDefaultGameMoviePlayer Suspend and Resume errors

#4.22

diff --git a/Engine/Source/Runtime/MoviePlayer/Public/MoviePlayer.h b/Engine/Source/Runtime/MoviePlayer/Public/MoviePlayer.h
--- a/Engine/Source/Runtime/MoviePlayer/Public/MoviePlayer.h
+++ b/Engine/Source/Runtime/MoviePlayer/Public/MoviePlayer.h
@@ -98,6 +98,9 @@ public:
 
 	void BroadcastCurrentMovieClipFinished(const FString& MovieClipThatFinished) { OnCurrentMovieClipFinished().Broadcast(MovieClipThatFinished); }
 
+	/** Called by the engine on suspend/resume */
+	virtual void Suspend() {};
+	virtual void Resume() {};
 
 };
 
@@ -241,6 +244,8 @@ public:
 	/** Check if the initial movie(s) is still playing */
 	virtual bool IsStartupMoviePlaying() const { return false; };
 	virtual void ForceCompletion() {};
+	virtual void Suspend() {};
+	virtual void Resume() {};
 };
 
 /** Creates the movie player */
