Fix for race condition startup movies due to invalidating a texture while getting its size.

diff --git a/Engine/Source/Runtime/MoviePlayer/Public/MoviePlayer.h b/Engine/Source/Runtime/MoviePlayer/Public/MoviePlayer.h
--- a/Engine/Source/Runtime/MoviePlayer/Public/MoviePlayer.h
+++ b/Engine/Source/Runtime/MoviePlayer/Public/MoviePlayer.h
@@ -34,7 +34,8 @@ public:
 	/* ISlateViewport interface. */
 	virtual FIntPoint GetSize() const override
 	{
-		return SlateTexture.IsValid() ? FIntPoint(SlateTexture.Pin()->GetWidth(), SlateTexture.Pin()->GetHeight()) : FIntPoint();
+		TSharedPtr<FSlateTexture2DRHIRef, ESPMode::ThreadSafe> SlateTextureSharedPtr = SlateTexture.Pin();
+		return SlateTextureSharedPtr.IsValid() ? FIntPoint(SlateTextureSharedPtr->GetWidth(), SlateTextureSharedPtr->GetHeight()) : FIntPoint();
 	}
 
 	virtual class FSlateShaderResource* GetViewportRenderTargetTexture() const override
