- Fix a crash that can happen if a background/foreground event comes in too early, with MoviePlayer trying to look at commandling

diff --git a/Engine/Source/Runtime/MoviePlayer/Private/MoviePlayer.cpp b/Engine/Source/Runtime/MoviePlayer/Private/MoviePlayer.cpp
--- a/Engine/Source/Runtime/MoviePlayer/Private/MoviePlayer.cpp
+++ b/Engine/Source/Runtime/MoviePlayer/Private/MoviePlayer.cpp
@@ -102,7 +102,7 @@ void DestroyMoviePlayer()
 
 IGameMoviePlayer* GetMoviePlayer()
 {
-	if (!IsMoviePlayerEnabled() || GUsingNullRHI)
+	if (!FCommandLine::IsInitialized() || !IsMoviePlayerEnabled() || GUsingNullRHI)
 	{
 		return FNullGameMoviePlayer::Get();
 	}
