Fixed a crash when trying to play a movie via FMoviePlayerWidgetRenderer such as the intro video when loading the ShooterGame sample game.

An ensure was added to SWidget in [REDACTED_HANDLE] to make sure that it's PaintParent (passed in via FPaintArgs) is not the same as the widget being painted.
This means we should now pass in nullptr for the VirtualRenderWindow widget rather than a reference to itself, mirroring the change made to PreLoadSlateThreading.cpp in the same CL.

diff --git a/Engine/Source/Runtime/MoviePlayer/Private/DefaultGameMoviePlayer.cpp b/Engine/Source/Runtime/MoviePlayer/Private/DefaultGameMoviePlayer.cpp
--- a/Engine/Source/Runtime/MoviePlayer/Private/DefaultGameMoviePlayer.cpp
+++ b/Engine/Source/Runtime/MoviePlayer/Private/DefaultGameMoviePlayer.cpp
@@ -856,7 +856,7 @@ void FMoviePlayerWidgetRenderer::DrawWindow(float DeltaTime)
 
 	int32 MaxLayerId = 0;
 	{
-		FPaintArgs PaintArgs(&VirtualRenderWindow.Get(), *HittestGrid, FVector2D::ZeroVector, FSlateApplication::Get().GetCurrentTime(), FSlateApplication::Get().GetDeltaTime());
+		FPaintArgs PaintArgs(nullptr, *HittestGrid, FVector2D::ZeroVector, FSlateApplication::Get().GetCurrentTime(), FSlateApplication::Get().GetDeltaTime());
 
 		// Paint the window
 		MaxLayerId = VirtualRenderWindow->Paint(
