Fix deprecation warnings in Movie Player

diff --git a/Engine/Source/Runtime/MoviePlayer/Private/DefaultGameMoviePlayer.cpp b/Engine/Source/Runtime/MoviePlayer/Private/DefaultGameMoviePlayer.cpp
--- a/Engine/Source/Runtime/MoviePlayer/Private/DefaultGameMoviePlayer.cpp
+++ b/Engine/Source/Runtime/MoviePlayer/Private/DefaultGameMoviePlayer.cpp
@@ -15,7 +15,7 @@
 #include "ShaderCompiler.h"
 #include "IHeadMountedDisplay.h"
 #include "IXRTrackingSystem.h"
-#include "IStereoLayers.h"
+#include "IXRLoadingScreen.h"
 #include "Misc/ConfigCacheIni.h"
 #include "HAL/FileManager.h"
 #include "Widgets/SVirtualWindow.h"
@@ -524,10 +524,10 @@ void FDefaultGameMoviePlayer::WaitForMovieToFinish(bool bAllowEngineTick)
 		LoadingIsDone.Set(1);
 		IsMoviePlaying = false;
 
-		IStereoLayers* StereoLayers;
-		if (GEngine && GEngine->StereoRenderingDevice.IsValid() && (StereoLayers = GEngine->StereoRenderingDevice->GetStereoLayers()) != nullptr && SyncMechanism == nullptr)
+		IXRLoadingScreen* LoadingScreen;
+		if (GEngine && GEngine->XRSystem.IsValid() && (LoadingScreen = GEngine->XRSystem->GetLoadingScreen()) != nullptr && SyncMechanism == nullptr)
 		{
-			StereoLayers->SetSplashScreenMovie(FTextureRHIRef());
+			LoadingScreen->ClearSplashes();
 		}
 
 		MovieStreamingIsDone.Set(1);
@@ -610,16 +610,24 @@ void FDefaultGameMoviePlayer::TickStreamer(float DeltaTime)
 			MovieStreamingIsDone.Set(1);
 		}
 
-		IStereoLayers* StereoLayers;
-		if (GEngine && GEngine->StereoRenderingDevice.IsValid() && (StereoLayers = GEngine->StereoRenderingDevice->GetStereoLayers()) != nullptr)
+		IXRLoadingScreen* LoadingScreen;
+		if (GEngine && GEngine->XRSystem.IsValid() && (LoadingScreen = GEngine->XRSystem->GetLoadingScreen()) != nullptr)
 		{
 			FTexture2DRHIRef Movie2DTexture = ActiveMovieStreamer->GetTexture();
-			FTextureRHIRef MovieTexture;
+			LoadingScreen->ClearSplashes();
 			if (Movie2DTexture.IsValid() && !bMovieIsDone)
 			{
-				MovieTexture = (FRHITexture*)Movie2DTexture.GetReference();
+				IXRLoadingScreen::FSplashDesc Splash;
+				Splash.Texture = (FRHITexture*)Movie2DTexture.GetReference();
+				Splash.bIsDynamic = true;
+				const FIntPoint TextureSize = Movie2DTexture->GetSizeXY();
+				const float	InvAspectRatio = (TextureSize.X > 0) ? float(TextureSize.Y) / float(TextureSize.X) : 1.0f;
+
+				Splash.bIgnoreAlpha = true;
+				Splash.Transform = FTransform(FVector(5.0f, 0.0f, 1.0f));
+				Splash.QuadSize = FVector2D(8.0f, 8.0f*InvAspectRatio);
+				LoadingScreen->AddSplash(Splash);
 			}
-			StereoLayers->SetSplashScreenMovie(MovieTexture);
 		}
 	}
 }
