MoviePlayer: Fixed assert trying to set a null viewport when there are no valid streamers.
#rnx

Merging
//UE5/Release-Engine-Staging/Engine/Source/Runtime/MoviePlayer/Private/DefaultGameMoviePlayer.cpp
to //UE5/Main/Engine/Source/Runtime/MoviePlayer/Private/DefaultGameMoviePlayer.cpp

diff --git a/Engine/Source/Runtime/MoviePlayer/Private/DefaultGameMoviePlayer.cpp b/Engine/Source/Runtime/MoviePlayer/Private/DefaultGameMoviePlayer.cpp
--- a/Engine/Source/Runtime/MoviePlayer/Private/DefaultGameMoviePlayer.cpp
+++ b/Engine/Source/Runtime/MoviePlayer/Private/DefaultGameMoviePlayer.cpp
@@ -376,7 +376,8 @@ bool FDefaultGameMoviePlayer::PlayMovie()
 			VirtualRenderWindow->SetContent(LoadingScreenContents.ToSharedRef());
 
 			// Register the movie viewport so that it can receive user input.
-			if (!FPlatformProperties::SupportsWindowedMode())
+			// There is only a valid viewport if we have a movie streamer as the streamer sets it.
+			if ((!FPlatformProperties::SupportsWindowedMode()) && (ActiveMovieStreamer.IsValid()))
 			{
 				TSharedPtr<SViewport> MovieViewport = MovieViewportWeakPtr.Pin();
 				if (MovieViewport.IsValid())
