MediaMovieStreamer: Fixed viewport size not being returned correctly when there is no texture.
#rnx

diff --git a/Engine/Plugins/Media/MediaMovieStreamer/Source/MediaMovieStreamer/Private/MediaMovieStreamer.cpp b/Engine/Plugins/Media/MediaMovieStreamer/Source/MediaMovieStreamer/Private/MediaMovieStreamer.cpp
--- a/Engine/Plugins/Media/MediaMovieStreamer/Source/MediaMovieStreamer/Private/MediaMovieStreamer.cpp
+++ b/Engine/Plugins/Media/MediaMovieStreamer/Source/MediaMovieStreamer/Private/MediaMovieStreamer.cpp
@@ -118,6 +118,19 @@ bool FMediaMovieStreamer::Init(const TArray<FString>& InMoviePaths, TEnumAsByte<
 	return true;
 }
 
+void FMediaMovieStreamer::PreviousViewportInterface(const TSharedPtr<ISlateViewport>& PreviousViewportInterface)
+{
+	// Use the size of the previous viewport as our default,
+	// as we may not have a size ourselves if we don't have a texture.
+	FIntPoint ViewportSize(ForceInitToZero);
+	if (PreviousViewportInterface.IsValid())
+	{
+		ViewportSize = PreviousViewportInterface->GetSize();
+	}
+
+	MovieViewport->SetDefaultSize(ViewportSize);
+}
+
 void FMediaMovieStreamer::ForceCompletion()
 {
 }
diff --git a/Engine/Plugins/Media/MediaMovieStreamer/Source/MediaMovieStreamer/Public/MediaMovieStreamer.h b/Engine/Plugins/Media/MediaMovieStreamer/Source/MediaMovieStreamer/Public/MediaMovieStreamer.h
--- a/Engine/Plugins/Media/MediaMovieStreamer/Source/MediaMovieStreamer/Public/MediaMovieStreamer.h
+++ b/Engine/Plugins/Media/MediaMovieStreamer/Source/MediaMovieStreamer/Public/MediaMovieStreamer.h
@@ -81,6 +81,7 @@ public:
 
 	/** IMovieStreamer interface */
 	virtual bool Init(const TArray<FString>& InMoviePaths, TEnumAsByte<EMoviePlaybackType> InPlaybackType) override;
+	virtual void PreviousViewportInterface(const TSharedPtr<ISlateViewport>& PreviousViewportInterface) override;
 	virtual void ForceCompletion() override;
 	virtual bool Tick(float DeltaTime) override;
 	virtual TSharedPtr<class ISlateViewport> GetViewportInterface() override;
diff --git a/Engine/Source/Runtime/MoviePlayer/Private/DefaultGameMoviePlayer.cpp b/Engine/Source/Runtime/MoviePlayer/Private/DefaultGameMoviePlayer.cpp
--- a/Engine/Source/Runtime/MoviePlayer/Private/DefaultGameMoviePlayer.cpp
+++ b/Engine/Source/Runtime/MoviePlayer/Private/DefaultGameMoviePlayer.cpp
@@ -390,6 +390,14 @@ bool FDefaultGameMoviePlayer::PlayMovie()
 				TSharedPtr<SViewport> MovieViewport = MovieViewportWeakPtr.Pin();
 				if (MovieViewport.IsValid())
 				{
+					// Let the streamer know about the previous viewport interface.
+					TSharedPtr<SViewport> GameViewport = FSlateApplication::Get().GetGameViewport();
+					if (GameViewport.IsValid())
+					{
+						TSharedPtr<ISlateViewport> ViewportInterface = GameViewport->GetViewportInterface().Pin();
+						ActiveMovieStreamer->PreviousViewportInterface(ViewportInterface);
+					}
+
 					FSlateApplication::Get().RegisterGameViewport(MovieViewport.ToSharedRef());
 				}
 			}
diff --git a/Engine/Source/Runtime/MoviePlayer/Public/MoviePlayer.h b/Engine/Source/Runtime/MoviePlayer/Public/MoviePlayer.h
--- a/Engine/Source/Runtime/MoviePlayer/Public/MoviePlayer.h
+++ b/Engine/Source/Runtime/MoviePlayer/Public/MoviePlayer.h
@@ -28,14 +28,14 @@ enum EMoviePlaybackType
 class FMovieViewport : public ISlateViewport, public TSharedFromThis<FMovieViewport>
 {
 public:
-	FMovieViewport() {}
+	FMovieViewport() : DefaultSize(ForceInitToZero) {}
 	~FMovieViewport() {}
 
 	/* ISlateViewport interface. */
 	virtual FIntPoint GetSize() const override
 	{
 		TSharedPtr<FSlateTexture2DRHIRef, ESPMode::ThreadSafe> SlateTextureSharedPtr = SlateTexture.Pin();
-		return SlateTextureSharedPtr.IsValid() ? FIntPoint(SlateTextureSharedPtr->GetWidth(), SlateTextureSharedPtr->GetHeight()) : FIntPoint();
+		return SlateTextureSharedPtr.IsValid() ? FIntPoint(SlateTextureSharedPtr->GetWidth(), SlateTextureSharedPtr->GetHeight()) : DefaultSize;
 	}
 
 	virtual class FSlateShaderResource* GetViewportRenderTargetTexture() const override
@@ -53,8 +53,17 @@ public:
 		SlateTexture = InTexture;
 	}
 
+	/** Sets the size to use when there is no texture. */
+	void SetDefaultSize(FIntPoint InSize)
+	{
+		DefaultSize = InSize;
+	}
+
 private:
 	TWeakPtr<FSlateTexture2DRHIRef, ESPMode::ThreadSafe> SlateTexture;
+
+	/** Size to use when we don't have a texture. */
+	FIntPoint DefaultSize;
 };
 
 
@@ -69,6 +78,9 @@ public:
 
 	virtual bool Init(const TArray<FString>& MoviePaths, TEnumAsByte<EMoviePlaybackType> inPlaybackType) = 0;
 	
+	/** Tells the movie streamer about the viewport interface that was active before us. */
+	virtual void PreviousViewportInterface(const TSharedPtr<ISlateViewport>& PreviousViewportInterface) {};
+
 	/** Forces the movie streamer to cancel what it's streaming and close. */
 	virtual void ForceCompletion() = 0;
 
