Don't parse the commandline every frame to check if the movie player should be disabled

diff --git a/Engine/Source/Runtime/MoviePlayer/Private/MoviePlayer.cpp b/Engine/Source/Runtime/MoviePlayer/Private/MoviePlayer.cpp
--- a/Engine/Source/Runtime/MoviePlayer/Private/MoviePlayer.cpp
+++ b/Engine/Source/Runtime/MoviePlayer/Private/MoviePlayer.cpp
@@ -122,7 +122,8 @@ bool IsMoviePlayerEnabled()
 	bool bEnabled = !GIsEditor && !IsRunningDedicatedServer() && !IsRunningCommandlet() && GUseThreadedRendering;
 
 #if !UE_BUILD_SHIPPING
-	bEnabled &= !FParse::Param(FCommandLine::Get(), TEXT("NoLoadingScreen"));
+	static bool bNoLoadingScreen = FParse::Param(FCommandLine::Get(), TEXT("NoLoadingScreen"));
+	bEnabled &= !bNoLoadingScreen;
 #endif
 
 	return bEnabled;
