Fix to Loading Screen hiding the User Widget before the end of the loading screen

diff --git a/Engine/Source/Runtime/MoviePlayer/Private/DefaultGameMoviePlayer.cpp b/Engine/Source/Runtime/MoviePlayer/Private/DefaultGameMoviePlayer.cpp
--- a/Engine/Source/Runtime/MoviePlayer/Private/DefaultGameMoviePlayer.cpp
+++ b/Engine/Source/Runtime/MoviePlayer/Private/DefaultGameMoviePlayer.cpp
@@ -460,19 +460,15 @@ void FDefaultGameMoviePlayer::WaitForMovieToFinish(bool bAllowEngineTick)
 			LoadingIsDone.Set(1);
 		}
 		
-        if (MainWindow.IsValid())
+        if (TSharedPtr<SWindow> MainWindowPtr = MainWindow.Pin())
         {
             // Transfer the content to the main window
-            MainWindow.Pin()->SetContent(LoadingScreenContents.ToSharedRef());
+			MainWindowPtr->SetContent(LoadingScreenContents.ToSharedRef());
         }
         if (VirtualRenderWindow.IsValid())
         {
             VirtualRenderWindow->SetContent(SNullWidget::NullWidget);
         }
-		if (UserWidgetHolder.IsValid())
-		{
-			UserWidgetHolder->SetContent(SNullWidget::NullWidget);
-		}
 
 		const bool bAutoCompleteWhenLoadingCompletes = LoadingScreenAttributes.bAutoCompleteWhenLoadingCompletes;
 		const bool bWaitForManualStop = LoadingScreenAttributes.bWaitForManualStop;
@@ -572,6 +568,10 @@ void FDefaultGameMoviePlayer::WaitForMovieToFinish(bool bAllowEngineTick)
 			}
 		}
 
+		if (UserWidgetHolder.IsValid())
+		{
+			UserWidgetHolder->SetContent(SNullWidget::NullWidget);
+		}
 		LoadingIsDone.Set(1);
 		IsMoviePlaying = false;
 		FCoreDelegates::OnAsyncLoadingFlushUpdate.Remove(OnAsyncLoadingFlushUpdateDelegateHandle);
