#ue5 - Type conversion fixes for modules: LiveLinkInterface, MainFrame, MediaAssets, MediaUtils, MergeActors, MoviePlayer.

diff --git a/Engine/Source/Editor/MainFrame/Private/MainFrameModule.cpp b/Engine/Source/Editor/MainFrame/Private/MainFrameModule.cpp
--- a/Engine/Source/Editor/MainFrame/Private/MainFrameModule.cpp
+++ b/Engine/Source/Editor/MainFrame/Private/MainFrameModule.cpp
@@ -219,7 +219,7 @@ void FMainFrameModule::CreateDefaultMainFrameAuxiliary(const bool bStartImmersiv
 			FDisplayMetrics DisplayMetrics;
 			FSlateApplication::Get().GetDisplayMetrics( DisplayMetrics );
 
-			const float DPIScale = FPlatformApplicationMisc::GetDPIScaleFactorAtPoint(DisplayMetrics.PrimaryDisplayWorkAreaRect.Left, DisplayMetrics.PrimaryDisplayWorkAreaRect.Top);
+			const float DPIScale = FPlatformApplicationMisc::GetDPIScaleFactorAtPoint((float)DisplayMetrics.PrimaryDisplayWorkAreaRect.Left, (float)DisplayMetrics.PrimaryDisplayWorkAreaRect.Top);
 
 			// Setup a position and size for the main frame window that's centered in the desktop work area
 			const float CenterScale = 0.65f;
diff --git a/Engine/Source/Editor/MainFrame/Private/Menus/SaveLayoutDialog.cpp b/Engine/Source/Editor/MainFrame/Private/Menus/SaveLayoutDialog.cpp
--- a/Engine/Source/Editor/MainFrame/Private/Menus/SaveLayoutDialog.cpp
+++ b/Engine/Source/Editor/MainFrame/Private/Menus/SaveLayoutDialog.cpp
@@ -275,7 +275,7 @@ void SSaveLayoutDialog::Construct(const FArguments& InArgs, const TSharedRef<FSa
 				[
 					// Constant height, whether the label is visible or not
 					SNew(SBox)
-					.HeightOverride(20)
+					.HeightOverride(20.f)
 					[
 						SNew(SBorder)
 						.Visibility( this, &SSaveLayoutDialog::GetNameErrorLabelVisibility )
diff --git a/Engine/Source/Editor/MergeActors/Private/MeshProxyTool/SMeshProxyDialog.cpp b/Engine/Source/Editor/MergeActors/Private/MeshProxyTool/SMeshProxyDialog.cpp
--- a/Engine/Source/Editor/MergeActors/Private/MeshProxyTool/SMeshProxyDialog.cpp
+++ b/Engine/Source/Editor/MergeActors/Private/MeshProxyTool/SMeshProxyDialog.cpp
@@ -426,12 +426,12 @@ void SThirdPartyMeshProxyDialog::HardAngleThresholdChanged(float NewValue)
 //Merge Distance
 TOptional<int32> SThirdPartyMeshProxyDialog::GetMergeDistance() const
 {
-	return Tool->ProxySettings.MergeDistance;
+	return UE::LWC::FloatToIntCastChecked<int32>(Tool->ProxySettings.MergeDistance);
 }
 
 void SThirdPartyMeshProxyDialog::MergeDistanceChanged(int32 NewValue)
 {
-	Tool->ProxySettings.MergeDistance = NewValue;
+	Tool->ProxySettings.MergeDistance = (float)NewValue;
 }
 
 //Texture Resolution
diff --git a/Engine/Source/Runtime/LiveLinkInterface/Public/LiveLinkTypes.h b/Engine/Source/Runtime/LiveLinkInterface/Public/LiveLinkTypes.h
--- a/Engine/Source/Runtime/LiveLinkInterface/Public/LiveLinkTypes.h
+++ b/Engine/Source/Runtime/LiveLinkInterface/Public/LiveLinkTypes.h
@@ -599,8 +599,8 @@ struct
 
 	FLiveLinkTimeCode& operator=(const FQualifiedFrameTime& InFrameTime)
 	{
-		const int32 NumberOfFramesInSecond = FMath::CeilToInt(InFrameTime.Rate.AsDecimal());
-		const int32 NumberOfFrames = FMath::RoundToZero(InFrameTime.Time.AsDecimal());
+		const int32 NumberOfFramesInSecond = FMath::CeilToInt32(InFrameTime.Rate.AsDecimal());
+		const int32 NumberOfFrames = (int32)(FMath::RoundToZero(InFrameTime.Time.AsDecimal()));
 
 		Seconds = (int32)FMath::RoundToZero(NumberOfFrames / (double)NumberOfFramesInSecond);
 		Frames = NumberOfFrames % NumberOfFramesInSecond;
diff --git a/Engine/Source/Runtime/MediaUtils/Private/MediaPlayerFacade.cpp b/Engine/Source/Runtime/MediaUtils/Private/MediaPlayerFacade.cpp
--- a/Engine/Source/Runtime/MediaUtils/Private/MediaPlayerFacade.cpp
+++ b/Engine/Source/Runtime/MediaUtils/Private/MediaPlayerFacade.cpp
@@ -2250,7 +2250,7 @@ void FMediaPlayerFacade::PostSampleProcessingTimeHandling(FTimespan DeltaTime)
 				if (Player->GetPlayerFeatureFlag(IMediaPlayer::EFeatureFlag::UseRealtimeWithVideoOnly))
 				{
 					double NewBaseTime = FPlatformTime::Seconds();
-					NextEstVideoTimeAtFrameStart.TimeStamp.Time += (NewBaseTime - NextEstVideoTimeAtFrameStart.SampledAtTime) * Rate;
+					NextEstVideoTimeAtFrameStart.TimeStamp.Time += FMath::TruncToInt64((NewBaseTime - NextEstVideoTimeAtFrameStart.SampledAtTime) * Rate);
 					NextEstVideoTimeAtFrameStart.SampledAtTime = NewBaseTime;
 				}
 				else
diff --git a/Engine/Source/Runtime/MediaUtils/Private/MediaRecorder.cpp b/Engine/Source/Runtime/MediaUtils/Private/MediaRecorder.cpp
--- a/Engine/Source/Runtime/MediaUtils/Private/MediaRecorder.cpp
+++ b/Engine/Source/Runtime/MediaUtils/Private/MediaRecorder.cpp
@@ -238,8 +238,8 @@ void FMediaRecorder::TickRecording()
 				const FIntPoint Size = Sample->GetDim();
 				EImagePixelType PixelType = EImagePixelType::Color;
 				ERGBFormat PixelLayout = ERGBFormat::BGRA;
-				int32 BitDepth = 8;
-				int32 NumChannels = 4;
+				uint8 BitDepth = 8;
+				uint8 NumChannels = 4;
 
 				if (Sample->GetFormat() == EMediaTextureSampleFormat::FloatRGBA)
 				{
diff --git a/Engine/Source/Runtime/MoviePlayer/Private/DefaultGameMoviePlayer.cpp b/Engine/Source/Runtime/MoviePlayer/Private/DefaultGameMoviePlayer.cpp
--- a/Engine/Source/Runtime/MoviePlayer/Private/DefaultGameMoviePlayer.cpp
+++ b/Engine/Source/Runtime/MoviePlayer/Private/DefaultGameMoviePlayer.cpp
@@ -661,7 +661,7 @@ void FDefaultGameMoviePlayer::BlockingTick()
 			HttpManager.Tick(0.0f);
 
 			// Callbacks.
-			OnMoviePlaybackTick().Broadcast(DeltaTime);
+			OnMoviePlaybackTick().Broadcast((float)DeltaTime);
 
 			UE_LOG(LogMoviePlayer, VeryVerbose, TEXT("BlockingTick deltatime:%f"), DeltaTime);
 		}
@@ -846,8 +846,8 @@ FVector2D FDefaultGameMoviePlayer::GetMovieSize() const
 	const FVector2D ScreenSize = MainWindow.Pin()->GetClientSizeInScreen();
 	if (MovieStreamingIsPrepared() && ActiveMovieStreamer.IsValid())
 	{
-		const float MovieAspectRatio = ActiveMovieStreamer->GetAspectRatio();
-		const float ScreenAspectRatio = ScreenSize.X / ScreenSize.Y;
+		const double MovieAspectRatio = ActiveMovieStreamer->GetAspectRatio();
+		const double ScreenAspectRatio = ScreenSize.X / ScreenSize.Y;
 		if (MovieAspectRatio < ScreenAspectRatio)
 		{
 			return FVector2D(ScreenSize.Y * MovieAspectRatio, ScreenSize.Y);
@@ -864,12 +864,12 @@ FVector2D FDefaultGameMoviePlayer::GetMovieSize() const
 
 FOptionalSize FDefaultGameMoviePlayer::GetMovieWidth() const
 {
-	return GetMovieSize().X;
+	return (float)GetMovieSize().X;
 }
 
 FOptionalSize FDefaultGameMoviePlayer::GetMovieHeight() const
 {
-	return GetMovieSize().Y;
+	return (float)GetMovieSize().Y;
 }
 
 EVisibility FDefaultGameMoviePlayer::GetSlateBackgroundVisibility() const
diff --git a/Engine/Source/Runtime/MoviePlayer/Private/MoviePlayerThreading.cpp b/Engine/Source/Runtime/MoviePlayer/Private/MoviePlayerThreading.cpp
--- a/Engine/Source/Runtime/MoviePlayer/Private/MoviePlayerThreading.cpp
+++ b/Engine/Source/Runtime/MoviePlayer/Private/MoviePlayerThreading.cpp
@@ -135,7 +135,7 @@ void FSlateLoadingSynchronizationMechanism::SlateThreadRunMainLoop()
 
 		if( TimeToWait > 0 )
 		{
-			FPlatformProcess::Sleep(TimeToWait);
+			FPlatformProcess::Sleep((float)TimeToWait);
 			CurrentTime = FPlatformTime::Seconds();
 			DeltaTime = CurrentTime - LastTime;
 		}
@@ -152,7 +152,7 @@ void FSlateLoadingSynchronizationMechanism::SlateThreadRunMainLoop()
 			FSlateRenderer* MainSlateRenderer = FSlateApplication::Get().GetRenderer();
 			FScopeLock ScopeLock(MainSlateRenderer->GetResourceCriticalSection());
 
-			WidgetRenderer->DrawWindow(DeltaTime);
+			WidgetRenderer->DrawWindow((float)DeltaTime);
 
 			SetSlateDrawPassEnqueued();
 
