MoviePlayer: Fixed a condition in the default movie player when set to play until user interaction that exited the playback loop when reaching the last element of the playlist without waiting for user interaction.

#rnx

diff --git a/Engine/Source/Runtime/MoviePlayer/Private/DefaultGameMoviePlayer.cpp b/Engine/Source/Runtime/MoviePlayer/Private/DefaultGameMoviePlayer.cpp
--- a/Engine/Source/Runtime/MoviePlayer/Private/DefaultGameMoviePlayer.cpp
+++ b/Engine/Source/Runtime/MoviePlayer/Private/DefaultGameMoviePlayer.cpp
@@ -34,7 +34,7 @@ class SDefaultMovieBorder : public SBorder
 {
 public:
 
-	SLATE_BEGIN_ARGS(SDefaultMovieBorder)		
+	SLATE_BEGIN_ARGS(SDefaultMovieBorder)
 		: _OnKeyDown()
 	{}
 
@@ -187,7 +187,7 @@ void FDefaultGameMoviePlayer::Initialize(FSlateRenderer& InSlateRenderer, TShare
 
 	// Add a delegate to start playing movies when we start loading a map
 	FCoreUObjectDelegates::PreLoadMap.AddRaw( this, &FDefaultGameMoviePlayer::OnPreLoadMap );
-	
+
 	// Shutdown the movie player if the app is exiting
 	FCoreDelegates::OnPreExit.AddRaw( this, &FDefaultGameMoviePlayer::Shutdown );
 
@@ -204,7 +204,7 @@ void FDefaultGameMoviePlayer::Initialize(FSlateRenderer& InSlateRenderer, TShare
 
 	WidgetRenderer = MakeShared<FMoviePlayerWidgetRenderer, ESPMode::ThreadSafe>(GameWindow, VirtualRenderWindow, &InSlateRenderer);
 
-	LoadingScreenContents = SNew(SDefaultMovieBorder)	
+	LoadingScreenContents = SNew(SDefaultMovieBorder)
 		.OnKeyDown(this, &FDefaultGameMoviePlayer::OnLoadingScreenKeyDown)
 		.OnMouseButtonDown(this, &FDefaultGameMoviePlayer::OnLoadingScreenMouseButtonDown)
 		[
@@ -341,7 +341,7 @@ bool FDefaultGameMoviePlayer::PlayMovie()
 {
 	bool bBeganPlaying = false;
 
-	// Allow systems to hook onto the movie player and provide loading screen data on demand 
+	// Allow systems to hook onto the movie player and provide loading screen data on demand
 	// if it has not been setup explicitly by the user.
 	if ( !LoadingScreenIsPrepared() )
 	{
@@ -352,7 +352,7 @@ bool FDefaultGameMoviePlayer::PlayMovie()
 	{
 		check(LoadingScreenAttributes.IsValid());
 		bUserCalledFinish = false;
-		
+
 		LastPlayTime = FPlatformTime::Seconds();
 
 		ActiveMovieStreamer.Reset();
@@ -379,7 +379,7 @@ bool FDefaultGameMoviePlayer::PlayMovie()
 			IsMoviePlaying = true;
 
 			UserWidgetDPIScaler->SetDPIScale(GetViewportDPIScale());
-			
+
 			UserWidgetHolder->SetContent(LoadingScreenAttributes.WidgetLoadingScreen.IsValid() ? LoadingScreenAttributes.WidgetLoadingScreen.ToSharedRef() : SNullWidget::NullWidget);
 			VirtualRenderWindow->Resize(MainWindow.Pin()->GetClientSizeInScreen());
 			VirtualRenderWindow->SetContent(LoadingScreenContents.ToSharedRef());
@@ -447,7 +447,7 @@ void FDefaultGameMoviePlayer::WaitForMovieToFinish(bool bAllowEngineTick)
 
 	if (LoadingScreenIsPrepared() && IsMovieCurrentlyPlaying())
 	{
-	
+
 		if (SyncMechanism)
 		{
 			SyncMechanism->DestroySlateThread();
@@ -460,7 +460,7 @@ void FDefaultGameMoviePlayer::WaitForMovieToFinish(bool bAllowEngineTick)
 		{
 			LoadingIsDone.Set(1);
 		}
-		
+
         if (TSharedPtr<SWindow> MainWindowPtr = MainWindow.Pin())
         {
             // Transfer the content to the main window
@@ -486,13 +486,14 @@ void FDefaultGameMoviePlayer::WaitForMovieToFinish(bool bAllowEngineTick)
 		// Don't continue playing on game shutdown
 		while ( !IsEngineExitRequested() &&
 				((bWaitForManualStop && !bUserCalledFinish)
-			||	(!bUserCalledFinish && !bEnforceMinimumTime && !IsMovieStreamingFinished() && !bAutoCompleteWhenLoadingCompletes) 
+			||	(!bUserCalledFinish && !bEnforceMinimumTime && !IsMovieStreamingFinished() && !bAutoCompleteWhenLoadingCompletes)
 			||	(bEnforceMinimumTime && (FPlatformTime::Seconds() - LastPlayTime) < LoadingScreenAttributes.MinimumLoadingScreenDisplayTime)))
 		{
 			BeginExitIfRequested();
 
-			// If we are in a loading loop, and this is the last movie in the playlist.. assume you can break out.
-			if (ActiveMovieStreamer.IsValid() && LoadingScreenAttributes.PlaybackType == MT_LoadingLoop && ActiveMovieStreamer->IsLastMovieInPlaylist())
+			// If we are in a loading loop, and this is the last movie in the playlist.. assume you can break out
+			// unless waiting for manual stop that hasn't been signaled yet.
+			if (ActiveMovieStreamer.IsValid() && LoadingScreenAttributes.PlaybackType == MT_LoadingLoop && ActiveMovieStreamer->IsLastMovieInPlaylist() && (!bWaitForManualStop || (bWaitForManualStop && bUserCalledFinish)))
 			{
 				break;
 			}
@@ -515,7 +516,7 @@ void FDefaultGameMoviePlayer::WaitForMovieToFinish(bool bAllowEngineTick)
 				// Gives widgets a chance to process any accumulated input
 				SlateApp.FinishedInputThisFrame();
 
-				float DeltaTime = SlateApp.GetDeltaTime();				
+				float DeltaTime = SlateApp.GetDeltaTime();
 
 				if (ActiveMovieStreamer.IsValid())
 				{
@@ -536,11 +537,11 @@ void FDefaultGameMoviePlayer::WaitForMovieToFinish(bool bAllowEngineTick)
 				ENQUEUE_RENDER_COMMAND(BeginLoadingMovieFrameAndTickMovieStreamer)(
 					[InMoviePlayer, DeltaTime](FRHICommandListImmediate& RHICmdList)
 					{
-						GFrameNumberRenderThread++;				
+						GFrameNumberRenderThread++;
 						InMoviePlayer->TickStreamer(RHICmdList, DeltaTime);
 					}
 				);
-				
+
 				{
 					FSlateRenderer* SlateRenderer = SlateApp.GetRenderer();
 					FScopeLock ScopeLock(SlateRenderer->GetResourceCriticalSection());
@@ -596,7 +597,7 @@ void FDefaultGameMoviePlayer::WaitForMovieToFinish(bool bAllowEngineTick)
 		{
 			ActiveMovieStreamer->Cleanup();
 		}
-	
+
 		// Finally, clear out the loading screen attributes, forcing users to always
 		// explicitly set the loading screen they want (rather than have stale loading screens)
 		LoadingScreenAttributes = FLoadingScreenAttributes();
@@ -604,7 +605,7 @@ void FDefaultGameMoviePlayer::WaitForMovieToFinish(bool bAllowEngineTick)
 		BroadcastMoviePlaybackFinished();
 	}
 	else
-	{	
+	{
 		UGameEngine* GameEngine = Cast<UGameEngine>(GEngine);
 
 		// Don't switch the window on game shutdown
@@ -653,7 +654,7 @@ void FDefaultGameMoviePlayer::BlockingTick()
 		{
 			// Yes. Time for another tick.
 			LastBlockingTickTime = Time;
-			
+
 			// Call HTTP manager.
 			FHttpManager& HttpManager = FHttpModule::Get().GetHttpManager();
 			HttpManager.Tick(0.0f);
@@ -678,7 +679,7 @@ void FDefaultGameMoviePlayer::BlockingFinished()
 			UE_LOG(LogMoviePlayer, Verbose, TEXT("BlockingFinished. Refcount: %d."), BlockingRefCount);
 			WaitForMovieToFinish();
 		}
-		
+
 		BlockingRefCount = 0;
 	}
 }
@@ -719,7 +720,7 @@ void FDefaultGameMoviePlayer::Tick( FRHICommandListImmediate& RHICmdList, float
 }
 
 void FDefaultGameMoviePlayer::TickStreamer(FRHICommandListImmediate& RHICmdList, float DeltaTime)
-{	
+{
 	if (MovieStreamingIsPrepared() && ActiveMovieStreamer.IsValid() && !IsMovieStreamingFinished())
 	{
 		const bool bMovieIsDone = ActiveMovieStreamer->Tick(RHICmdList, DeltaTime);
@@ -988,7 +989,7 @@ void FMoviePlayerWidgetRenderer::DrawWindow(float DeltaTime)
 	{
 		Scale *= MainWindow->GetDPIScaleFactor();
 	}
-	
+
 	FVector2D DrawSize = VirtualRenderWindow->GetClientSizeInScreen() / Scale;
 
 	FSlateApplication::Get().Tick(ESlateTickType::Time);
