Call OnSamplingInput delegate while waiting for a movie to finish.

A normal engine tick (FEngineLoop::Tick) invokes the OnSamplingInput delegate before calling PollGameDeviceState. The loop within FDefaultGameMoviePlayer::WaitForMovieToFinish did not.

Some platforms use OnSamplingInput to respond to system notifications that a new input device has been attached. Without this call, splash screen movies could not be skipped because no input devices were registered until after the splash screens finished when the engine had its first tick.

#rnx

diff --git a/Engine/Source/Runtime/MoviePlayer/Private/DefaultGameMoviePlayer.cpp b/Engine/Source/Runtime/MoviePlayer/Private/DefaultGameMoviePlayer.cpp
--- a/Engine/Source/Runtime/MoviePlayer/Private/DefaultGameMoviePlayer.cpp
+++ b/Engine/Source/Runtime/MoviePlayer/Private/DefaultGameMoviePlayer.cpp
@@ -512,6 +512,8 @@ void FDefaultGameMoviePlayer::WaitForMovieToFinish(bool bAllowEngineTick)
 
 				FPlatformApplicationMisc::PumpMessages(true);
 
+				FCoreDelegates::OnSamplingInput.Broadcast();
+
 				SlateApp.PollGameDeviceState();
 				// Gives widgets a chance to process any accumulated input
 				SlateApp.FinishedInputThisFrame();
