Fix FDefaultGameMoviePlayer::WaitForMovieToFinish(...) not incrementing GFrameCounterRenderThread between GEngine->Tick(...) calls.
- caused check in FRayTracingGeometryManager::Tick(...) to fail.

#Synced-CL 51376035

diff --git a/Engine/Source/Runtime/MoviePlayer/Private/DefaultGameMoviePlayer.cpp b/Engine/Source/Runtime/MoviePlayer/Private/DefaultGameMoviePlayer.cpp
--- a/Engine/Source/Runtime/MoviePlayer/Private/DefaultGameMoviePlayer.cpp
+++ b/Engine/Source/Runtime/MoviePlayer/Private/DefaultGameMoviePlayer.cpp
@@ -535,10 +535,13 @@ void FDefaultGameMoviePlayer::WaitForMovieToFinish(bool bAllowEngineTick)
 					ActiveMovieStreamer->TickPostEngine();
 				}
 
+				GFrameCounter++;
+
 				FDefaultGameMoviePlayer* InMoviePlayer = this;
 				ENQUEUE_RENDER_COMMAND(BeginLoadingMovieFrameAndTickMovieStreamer)(
-					[InMoviePlayer, DeltaTime](FRHICommandListImmediate& RHICmdList)
+					[InMoviePlayer, DeltaTime, CurrentFrameCounter = GFrameCounter](FRHICommandListImmediate& RHICmdList)
 					{
+						GFrameCounterRenderThread = CurrentFrameCounter;
 						GFrameNumberRenderThread++;
 						InMoviePlayer->TickStreamer(RHICmdList, DeltaTime);
 					}
